[Mutable] Program cache rewrok.

diff --git a/Engine/Plugins/Mutable/Source/MutableRuntime/Private/MuR/Model.cpp b/Engine/Plugins/Mutable/Source/MutableRuntime/Private/MuR/Model.cpp
--- a/Engine/Plugins/Mutable/Source/MutableRuntime/Private/MuR/Model.cpp
+++ b/Engine/Plugins/Mutable/Source/MutableRuntime/Private/MuR/Model.cpp
@@ -39,7 +39,6 @@ namespace UE::Mutable::Private
 		// Add the null instruction at address 0.
 		// TODO: Will do it in the linker
 		AppendCode( ByteCode, EOpType::NONE );
-		OpAddress.Add(0);
 	}
 
 
@@ -63,7 +62,7 @@ namespace UE::Mutable::Private
 	}
 
 
-	uint64 FProgram::FState::IsDynamic(OP::ADDRESS At) const
+	uint64 FProgram::FState::IsDynamic(FOperation::ADDRESS At) const
 	{
 		uint64 Result = 0;
 
@@ -80,7 +79,7 @@ namespace UE::Mutable::Private
 	}
 
 
-	bool FProgram::FState::IsUpdateCache(OP::ADDRESS At) const
+	bool FProgram::FState::IsUpdateCache(FOperation::ADDRESS At) const
 	{
 		bool Result = false;
 
@@ -96,7 +95,7 @@ namespace UE::Mutable::Private
 	}
 
 
-	void FProgram::FState::AddUpdateCache(OP::ADDRESS At)
+	void FProgram::FState::AddUpdateCache(FOperation::ADDRESS At)
 	{
 		if (!IsUpdateCache(At))
 		{
@@ -107,7 +106,6 @@ namespace UE::Mutable::Private
 
 	void FProgram::Serialise(FOutputArchive& Arch) const
 	{
-		Arch << OpAddress;
 		Arch << ByteCode;
 		Arch << States;
 		Arch << Roms;
@@ -139,7 +137,6 @@ namespace UE::Mutable::Private
 
 	void FProgram::Unserialise(FInputArchive& Arch)
 	{
-		Arch >> OpAddress;
 		Arch >> ByteCode;
 		Arch >> States;
 		Arch >> Roms;
@@ -224,24 +221,24 @@ namespace UE::Mutable::Private
     }
 	
 #if WITH_EDITOR
-    OP::ADDRESS FProgram::AddConstant(const TManagedPtr<const FLayout>& Layout)
+    FOperation::ADDRESS FProgram::AddConstant(const TManagedPtr<const FLayout>& Layout)
     {
 	    // Ensure unique
 	    for (int32 Index = 0; Index < ConstantLayouts.Num(); ++Index)
 	    {
 		    if (ConstantLayouts[Index] == Layout)
 		    {
-			    return (OP::ADDRESS)Index;
+			    return (FOperation::ADDRESS)Index;
 		    }
 	    }
 
-	    OP::ADDRESS Index = OP::ADDRESS(ConstantLayouts.Num());
+	    FOperation::ADDRESS Index = FOperation::ADDRESS(ConstantLayouts.Num());
 	    ConstantLayouts.Add(Layout);
 	    return Index;
     }
 
 
-	OP::ADDRESS FProgram::AddConstant(const TManagedPtr<const FSkeleton>& InSkeleton)
+	FOperation::ADDRESS FProgram::AddConstant(const TManagedPtr<const FSkeleton>& InSkeleton)
 	{
 		// Ensure unique
 		for (int32 i = 0; i < ConstantSkeletons.Num(); ++i)
@@ -251,7 +248,7 @@ namespace UE::Mutable::Private
 				*ConstantSkeletons[i] == *InSkeleton // Does not compare BoneIds
 				)
 			{
-				return (OP::ADDRESS)i;
+				return (FOperation::ADDRESS)i;
 			}
 		}
 
@@ -269,22 +266,22 @@ namespace UE::Mutable::Private
 			BoneIds[Index] = AddConstant(BoneNames[Index]);
 		}
 
-		return OP::ADDRESS(ConstantSkeletons.Add(Skeleton));
+		return FOperation::ADDRESS(ConstantSkeletons.Add(Skeleton));
 	}
 
 
-    OP::ADDRESS FProgram::AddConstant(const TManagedPtr<const FPhysicsBody>& PhysicsBody)
+    FOperation::ADDRESS FProgram::AddConstant(const TManagedPtr<const FPhysicsBody>& PhysicsBody)
     {
 	    // Ensure unique
 	    for (int32 Index = 0; Index < ConstantPhysicsBodies.Num(); ++Index)
 	    {
 		    if (ConstantPhysicsBodies[Index] == PhysicsBody || *ConstantPhysicsBodies[Index] == *PhysicsBody )
 		    {
-			    return (OP::ADDRESS)Index;
+			    return (FOperation::ADDRESS)Index;
 		    }
 	    }
 
-	    OP::ADDRESS index = OP::ADDRESS(ConstantPhysicsBodies.Num());
+	    FOperation::ADDRESS index = FOperation::ADDRESS(ConstantPhysicsBodies.Num());
 	    ConstantPhysicsBodies.Add(PhysicsBody);
 	    return index;
     }
@@ -330,7 +327,7 @@ namespace UE::Mutable::Private
 	}
 
 
-	OP::CONSTANT_LIST_UINT32_INDEX FProgram::AddConstant(const TArray<FMeshSocket>& Sockets)
+	FOperation::CONSTANT_LIST_UINT32_INDEX FProgram::AddConstant(const TArray<FMeshSocket>& Sockets)
 	{
 		TArray<uint32> SocketIds;
 		SocketIds.Reset(Sockets.Num());
@@ -344,9 +341,9 @@ namespace UE::Mutable::Private
 		return AddConstant(SocketIds);
 	}
 
-    OP::CONSTANT_STRING_INDEX FProgram::AddConstant(const FString& Str)
+    FOperation::CONSTANT_STRING_INDEX FProgram::AddConstant(const FString& Str)
     {     
-	    return (OP::CONSTANT_STRING_INDEX)ConstantStrings.AddUnique(Str);
+	    return (FOperation::CONSTANT_STRING_INDEX)ConstantStrings.AddUnique(Str);
     }
 
 
@@ -367,7 +364,7 @@ namespace UE::Mutable::Private
 	}
 
 
-	OP::ADDRESS FProgram::AddConstant(const TArray<FName>& NameList)
+	FOperation::ADDRESS FProgram::AddConstant(const TArray<FName>& NameList)
 	{
 		const int32 NumNames = NameList.Num();
 
@@ -379,23 +376,23 @@ namespace UE::Mutable::Private
 			NameAddrs[NameIndex] = AddConstant(NameList[NameIndex]);
 		}
 
-		return (OP::ADDRESS)AddConstant(NameAddrs);
+		return (FOperation::ADDRESS)AddConstant(NameAddrs);
 	}
 
 
-    OP::ADDRESS FProgram::AddConstant(const TArray<uint32>& UInt32List)
+    FOperation::ADDRESS FProgram::AddConstant(const TArray<uint32>& UInt32List)
     {
-	    return (OP::ADDRESS)ConstantUInt32Lists.AddUnique(UInt32List);
+	    return (FOperation::ADDRESS)ConstantUInt32Lists.AddUnique(UInt32List);
     }
 
 
-    OP::ADDRESS FProgram::AddConstant(const TArray<uint64>& UInt64List)
+    FOperation::ADDRESS FProgram::AddConstant(const TArray<uint64>& UInt64List)
     {
-	    return (OP::ADDRESS)ConstantUInt64Lists.AddUnique(UInt64List);
+	    return (FOperation::ADDRESS)ConstantUInt64Lists.AddUnique(UInt64List);
     }
 
 
-    OP::ADDRESS FProgram::AddConstant(const TArray<FString>& StringList)
+    FOperation::ADDRESS FProgram::AddConstant(const TArray<FString>& StringList)
     {
 	    const int32 NumStrings = StringList.Num();
 
@@ -407,115 +404,117 @@ namespace UE::Mutable::Private
 		    StringAddrs[StringIndex] = AddConstant(StringList[StringIndex]);
 	    }
 
-	    return (OP::ADDRESS)AddConstant(StringAddrs);
+	    return (FOperation::ADDRESS)AddConstant(StringAddrs);
     }
 
 
-    OP::CONSTANT_MATRIX_INDEX FProgram::AddConstant(const FMatrix44f& Constant)
+    FOperation::CONSTANT_MATRIX_INDEX FProgram::AddConstant(const FMatrix44f& Constant)
     {
 	    // Ensure unique
 	    for (int32 i = 0; i < ConstantMatrices.Num(); ++i)
 	    {
 		    if (ConstantMatrices[i] == Constant)
 		    {
-			    return (OP::CONSTANT_MATRIX_INDEX)i;
+			    return (FOperation::CONSTANT_MATRIX_INDEX)i;
 		    }
 	    }
 
-	    OP::CONSTANT_MATRIX_INDEX Index = OP::CONSTANT_MATRIX_INDEX(ConstantMatrices.Num());
+	    FOperation::CONSTANT_MATRIX_INDEX Index = FOperation::CONSTANT_MATRIX_INDEX(ConstantMatrices.Num());
 	    ConstantMatrices.Add(Constant);
 	    return Index;
     }
 
 
-    OP::CONSTANT_SHAPE_INDEX FProgram::AddConstant(const FShape& Constant)
+    FOperation::CONSTANT_SHAPE_INDEX FProgram::AddConstant(const FShape& Constant)
     {
 	    // Ensure unique
 	    for (int32 Index = 0; Index<ConstantShapes.Num(); ++Index)
 	    {
 		    if (ConstantShapes[Index] == Constant)
 		    {
-			    return (OP::CONSTANT_SHAPE_INDEX)Index;
+			    return (FOperation::CONSTANT_SHAPE_INDEX)Index;
 		    }
 	    }
 
-	    OP::CONSTANT_SHAPE_INDEX Index = OP::CONSTANT_SHAPE_INDEX(ConstantShapes.Num());
+	    FOperation::CONSTANT_SHAPE_INDEX Index = FOperation::CONSTANT_SHAPE_INDEX(ConstantShapes.Num());
 	    ConstantShapes.Add(Constant);
 	    return Index;
     }
 
 
-    OP::ADDRESS FProgram::AddConstant(const FProjector& Constant)
+    FOperation::ADDRESS FProgram::AddConstant(const FProjector& Constant)
     {
 	    // Ensure unique
 	    for (int32 Index = 0; Index < ConstantProjectors.Num(); ++Index)
 	    {
 		    if (ConstantProjectors[Index] == Constant)
 		    {
-			    return (OP::ADDRESS)Index;
+			    return (FOperation::ADDRESS)Index;
 		    }
 	    }
 
-	    OP::ADDRESS Index = OP::ADDRESS(ConstantProjectors.Num());
+	    FOperation::ADDRESS Index = FOperation::ADDRESS(ConstantProjectors.Num());
 	    ConstantProjectors.Add(Constant);
 	    return Index;
     }
 
 
-    OP::ADDRESS FProgram::AddConstant(const FRichCurve& Constant)
+    FOperation::ADDRESS FProgram::AddConstant(const FRichCurve& Constant)
     {
-	    OP::ADDRESS Index = OP::ADDRESS(ConstantCurves.AddUnique(Constant ));
+	    FOperation::ADDRESS Index = FOperation::ADDRESS(ConstantCurves.AddUnique(Constant ));
 	    return Index;
     }
 
 
-    OP::ADDRESS FProgram::AddConstant(const TManagedPtr<const FMaterial>& Material)
+    FOperation::ADDRESS FProgram::AddConstant(const TManagedPtr<const FMaterial>& Material)
     {
 	    for (int32 Index=0; Index<ConstantMaterials.Num(); ++Index)
 	    {
 		    if (ConstantMaterials[Index] == Material || *ConstantMaterials[Index] == *Material)
 		    {
-			    return static_cast<OP::ADDRESS>(Index);
+			    return static_cast<FOperation::ADDRESS>(Index);
 		    }
 	    }
 
-	    OP::ADDRESS Index = static_cast<OP::ADDRESS>(ConstantMaterials.Num());
+	    FOperation::ADDRESS Index = static_cast<FOperation::ADDRESS>(ConstantMaterials.Num());
 	    ConstantMaterials.Add(Material);
 	    return Index;
     }
 #endif // WITH_EDITOR
 	
 
-    EOpType FProgram::GetOpType(OP::ADDRESS At) const
+    EOpType FProgram::GetOpType(FOperation::ADDRESS Address) const
     {
-	    if (At >= OP::ADDRESS(OpAddress.Num()))
-	    {
-		    return EOpType::NONE;
-	    }
-		
-	    EOpType Result;
-	    uint64 byteCodeAddress = OpAddress[At];
-	    FMemory::Memcpy(&Result, &ByteCode[byteCodeAddress], sizeof(EOpType));
-	    check(Result < EOpType::COUNT);
-	    return Result;
+		uint32 ByteCodeOffset = GetAddressByteCodeOffset(Address);
+
+		if ((ByteCodeOffset >= (uint32)ByteCode.Num()) | 
+			(ByteCodeOffset >= FOperation::MaxAddressByteCodeOffset)) 
+		{
+			return EOpType::NONE;
+		}
+
+		EOpType Result;
+		FMemory::Memcpy(&Result, &ByteCode[ByteCodeOffset], sizeof(EOpType));
+
+		check(Result < EOpType::COUNT);
+
+		return Result;
     }
 
 
-    const uint8* FProgram::GetOpArgsPointer(OP::ADDRESS At) const
+    const uint8* FProgram::GetOpArgsPointer(FOperation::ADDRESS Address) const
     {
-	    uint64 byteCodeAddress = OpAddress[At];
-	    byteCodeAddress += sizeof(EOpType);
-	    const uint8* pData = (const uint8*)&ByteCode[byteCodeAddress];
-	    return pData;
+	    uint32 ByteCodeOffset = GetAddressByteCodeOffset(Address) + sizeof(EOpType);
+	    const uint8* DataPtr = (const uint8*)&ByteCode[ByteCodeOffset];
+	    return DataPtr;
     }
 
 
-    uint8* FProgram::GetOpArgsPointer(OP::ADDRESS At)
+    uint8* FProgram::GetOpArgsPointer(FOperation::ADDRESS Address)
     {
-	    uint64 byteCodeAddress = OpAddress[At];
-	    byteCodeAddress += sizeof(EOpType);
-	    uint8* pData = (uint8*)&ByteCode[byteCodeAddress];
-	    return pData;
+	    uint32 ByteCodeOffset = GetAddressByteCodeOffset(Address) + sizeof(EOpType);
+	    uint8* DataPtr = &ByteCode[ByteCodeOffset];
+	    return DataPtr;
     }
 
 	
diff --git a/Engine/Plugins/Mutable/Source/MutableRuntime/Private/MuR/Operations.cpp b/Engine/Plugins/Mutable/Source/MutableRuntime/Private/MuR/Operations.cpp
--- a/Engine/Plugins/Mutable/Source/MutableRuntime/Private/MuR/Operations.cpp
+++ b/Engine/Plugins/Mutable/Source/MutableRuntime/Private/MuR/Operations.cpp
@@ -10,175 +10,10 @@
 
 namespace UE::Mutable::Private
 {
-    MUTABLE_IMPLEMENT_POD_SERIALISABLE(OP);
-    MUTABLE_IMPLEMENT_POD_VECTOR_SERIALISABLE(OP);
-
-    // clang-format off
-    static const FOpDesc SOperationRuntimeDescs[] =
-	{ 
-		// DataType				
-        { EDataType::None,			},	// NONE
-
-		{ EDataType::Bool,			},	// BO_CONSTANT
-		{ EDataType::Int,			},	// NU_CONSTANT
-		{ EDataType::Scalar,			},	// SC_CONSTANT
-		{ EDataType::Color,			},	// CO_CONSTANT
-		{ EDataType::Image,			},	// IM_CONSTANT
-		{ EDataType::Mesh,			},	// ME_CONSTANT
-		{ EDataType::Layout,			},	// LA_CONSTANT
-		{ EDataType::Projector,		},	// PR_CONSTANT
-		{ EDataType::String,			},	// ST_CONSTANT
-		{ EDataType::ExtensionData,	},	// ED_CONSTANT
-		{ EDataType::Matrix,			},	// MA_CONSTANT
-		{ EDataType::Material,		},	// MI_CONSTANT
-
-		{ EDataType::Bool,			},	// BO_PARAMETER
-		{ EDataType::Int,			},	// NU_PARAMETER
-		{ EDataType::Scalar,			},	// SC_PARAMETER
-		{ EDataType::Color,			},	// CO_PARAMETER
-		{ EDataType::Projector,		},	// PR_PARAMETER
-		{ EDataType::Image,			},	// IM_PARAMETER
-		{ EDataType::SkeletalMesh,	},	// SK_PARAMETER
-		{ EDataType::String,			},	// ST_PARAMETER
-		{ EDataType::Matrix,			},	// MA_PARAMETER
-		{ EDataType::Material,		},	// MI_PARAMETER
-		{ EDataType::InstancedStruct,},	// MI_PARAMETER
-
-		{ EDataType::Image,			},	// IM_REFERENCE
-		{ EDataType::Mesh,			},	// ME_REFERENCE
-
-		{ EDataType::Int,			},	// NU_CONDITIONAL
-		{ EDataType::Scalar,			},	// SC_CONDITIONAL
-		{ EDataType::Color,			},	// CO_CONDITIONAL
-		{ EDataType::Image,			},	// IM_CONDITIONAL
-		{ EDataType::Mesh,			},	// ME_CONDITIONAL
-		{ EDataType::Layout,			},	// LA_CONDITIONAL
-		{ EDataType::Instance,		},	// IN_CONDITIONAL
-		{ EDataType::ExtensionData,	},	// ED_CONDITIONAL
-		{ EDataType::Material,		},	// MI_CONDITIONAL
-
-		{ EDataType::Int,			},	// NU_SWITCH
-		{ EDataType::Scalar,			},	// SC_SWITCH
-		{ EDataType::Color,			},	// CO_SWITCH
-		{ EDataType::Image,			},	// IM_SWITCH
-		{ EDataType::Mesh,			},	// ME_SWITCH
-		{ EDataType::Layout,			},	// LA_SWITCH
-		{ EDataType::Instance,		},	// IN_SWITCH
-		{ EDataType::ExtensionData,	},	// ED_SWITCH
-		{ EDataType::Material,		},	// MI_SWITCH
-
-		{ EDataType::Scalar,			},	// SC_MATERIAL_BREAK
-		{ EDataType::Color,			},	// CO_MATERIAL_BREAK
-		{ EDataType::Image,			},	// IM_MATERIAL_BREAK
-		{ EDataType::Image,			},	// IM_PARAMETER_FROM_MATERIAL
-		{ EDataType::Material,		},	// MI_MODIFY
-
-		{ EDataType::Bool,			},	// BO_EQUAL_SC_CONST
-		{ EDataType::Bool,			},	// BO_AND
-		{ EDataType::Bool,			},	// BO_OR
-		{ EDataType::Bool,			},	// BO_NOT
-
-		{ EDataType::Scalar,			},	// SC_ARITHMETIC
-		{ EDataType::Scalar,			},	// SC_CURVE
-		{ EDataType::Scalar,			},	// SC_EXTERNAL
-
-		{ EDataType::Color,			},	// CO_SAMPLEIMAGE
-		{ EDataType::Color,			},	// CO_SWIZZLE
-		{ EDataType::Color,			},	// CO_FROMSCALARS
-		{ EDataType::Color,			},	// CO_ARITHMETIC
-		{ EDataType::Color,			},	// CO_LINEARTOSRGB
-		{ EDataType::Color,			},	// CO_EXTERNAL
-
-		{ EDataType::Image,			},	// IM_LAYER
-		{ EDataType::Image,			},	// IM_LAYERCOLOUR
-		{ EDataType::Image,			},	// IM_PIXELFORMAT
-		{ EDataType::Image,			},	// IM_MIPMAP
-		{ EDataType::Image,			},	// IM_RESIZE
-		{ EDataType::Image,			},	// IM_RESIZELIKE
-		{ EDataType::Image,			},	// IM_RESIZEREL
-		{ EDataType::Image,			},	// IM_BLANKLAYOUT
-		{ EDataType::Image,			},	// IM_COMPOSE
-		{ EDataType::Image,			},	// IM_MULTICOMPOSE
-		{ EDataType::Image,			},	// IM_INTERPOLATE
-		{ EDataType::Image,			},	// IM_SATURATE
-		{ EDataType::Image,			},	// IM_LUMINANCE
-		{ EDataType::Image,			},	// IM_SWIZZLE
-		{ EDataType::Image,			},	// IM_COLOURMAP
-		{ EDataType::Image,			},	// IM_BINARISE
-		{ EDataType::Image,			},	// IM_PLAINCOLOUR
-		{ EDataType::Image,			},	// IM_CROP
-		{ EDataType::Image,			},	// IM_PATCH
-		{ EDataType::Image,			},	// IM_RASTERMESH
-		{ EDataType::Image,			},	// IM_MAKEGROWMAP
-		{ EDataType::Image,			},	// IM_DISPLACE
-		{ EDataType::Image,			},	// IM_MULTILAYER
-		{ EDataType::Image,			},	// IM_INVERT
-		{ EDataType::Image,			},	// IM_NORMALCOMPOSITE
-		{ EDataType::Image,			},	// IM_TRANSFORM
-		{ EDataType::Image,			},	// IM_PARAMETER_CONVERT
-		{ EDataType::Image,			},	// IM_EXTERNAL
-
-		{ EDataType::Mesh,			},	// ME_APPLYLAYOUT
-		{ EDataType::Mesh,			},	// ME_PREPARELAYOUT
-		{ EDataType::Mesh,			},	// ME_DIFFERENCE
-		{ EDataType::Mesh,			},	// ME_MORPH
-		{ EDataType::Mesh,			},	// ME_MERGE
-		{ EDataType::Mesh,			},	// ME_MASKCLIPMESH
-		{ EDataType::Mesh,			},	// ME_MASKCLIPUVMASK
-		{ EDataType::Mesh,			},	// ME_MASKDIFF
-		{ EDataType::Mesh,			},	// ME_REMOVEMASK
-		{ EDataType::Mesh,			},	// ME_FORMAT
-		{ EDataType::Mesh,			},	// ME_EXTRACTLAYOUTBLOCK
-		{ EDataType::Mesh,			},	// ME_TRANSFORM
-		{ EDataType::Mesh,			},	// ME_CLIPMORPHPLANE
-		{ EDataType::Mesh,			},	// ME_CLIPWITHMESH
-		{ EDataType::Mesh,			},	// ME_PROJECT
-		{ EDataType::Mesh,			},	// ME_APPLYPOSE
-		{ EDataType::Mesh,			},	// ME_BINDSHAPE
-		{ EDataType::Mesh,			},	// ME_APPLYSHAPE
-		{ EDataType::Mesh,			},	// ME_CLIPDEFORM
-		{ EDataType::Mesh,			},	// ME_MORPHRESHAPE
-		{ EDataType::Mesh,			},	// ME_OPTIMIZESKINNING
-		{ EDataType::Mesh,			},	// ME_ADDMETADATA
-		{ EDataType::Mesh,			},	// ME_TRANSFORMWITHMESH
-		{ EDataType::Mesh,			},	// ME_TRANSFORMWITHBONE
-		{ EDataType::Mesh,			},	// ME_EXTERNAL
-		{ EDataType::Mesh,			},	// ME_SKELETALMESH_BREAK
-
-		{ EDataType::Instance,		},	// IN_ADDMESH
-		{ EDataType::Instance,		},	// IN_ADDSTRING
-		{ EDataType::Instance,		},	// IN_ADDSURFACE
-		{ EDataType::Instance,		},	// IN_ADDCOMPONENT
-		{ EDataType::Instance,		},	// IN_ADDLOD
-		{ EDataType::Instance,		},	// IN_ADDSKELETALMESH
-		{ EDataType::Instance,		},	// IN_ADDEXTENSIONDATA
-		{ EDataType::Instance,		},	// IN_ADDOVERLAYMATERIAL
-		{ EDataType::Instance,		},	// IN_ADDOVERRIDEMATERIAL
-		{ EDataType::Instance,		},	// IN_ADDMATERIAL
-
-		{ EDataType::Layout,			},	// LA_PACK
-		{ EDataType::Layout,			},	// LA_MERGE
-		{ EDataType::Layout,			},	// LA_REMOVEBLOCKS
-		{ EDataType::Layout,			},	// LA_FROMMESH
-
-		{ EDataType::Material,		},	// MI_EXTERNAL
-		{ EDataType::Material,		},	// MI_SKELETALMESH_BREAK
-
-    	{ EDataType::InstancedStruct,},	// IS_EXTERNAL
-	};
-
-    // clang-format on
-
-    static_assert(UE_ARRAY_COUNT(SOperationRuntimeDescs) == int32(EOpType::COUNT), "OperationDescMismatch");
-
-	const FOpDesc& GetOpDesc( EOpType type )
-	{
-        return SOperationRuntimeDescs[ (int32)type ];
-	}
-
-
-    //---------------------------------------------------------------------------------------------
-    void ForEachReference( const FProgram& program, OP::ADDRESS at, const TFunctionRef<void(OP::ADDRESS)> f )
+    MUTABLE_IMPLEMENT_POD_SERIALISABLE(FOperation);
+    MUTABLE_IMPLEMENT_POD_VECTOR_SERIALISABLE(FOperation);
+
+    void ForEachReference( const FProgram& program, FOperation::ADDRESS at, const TFunctionRef<void(FOperation::ADDRESS)> f )
     {
         EOpType type = program.GetOpType(at);
         switch ( type )
@@ -213,7 +48,7 @@ namespace UE::Mutable::Private
 
         case EOpType::SC_CURVE:
         {
-			OP::ScalarCurveArgs args = program.GetOpArgs<OP::ScalarCurveArgs>(at);
+			FOperation::ScalarCurveArgs args = program.GetOpArgs<FOperation::ScalarCurveArgs>(at);
             f(args.time );
             break;
         }
@@ -228,7 +63,7 @@ namespace UE::Mutable::Private
 		case EOpType::ED_CONDITIONAL:
 		case EOpType::MI_CONDITIONAL:
         {
-			OP::ConditionalArgs args = program.GetOpArgs<OP::ConditionalArgs>(at);
+			FOperation::ConditionalArgs args = program.GetOpArgs<FOperation::ConditionalArgs>(at);
             f(args.condition );
             f(args.yes );
             f(args.no );
@@ -247,17 +82,17 @@ namespace UE::Mutable::Private
         {
 			const uint8* Data = program.GetOpArgsPointer(at);
 			
-			OP::ADDRESS VarAddress;
-			FMemory::Memcpy(&VarAddress, Data, sizeof(OP::ADDRESS));
-			Data += sizeof(OP::ADDRESS);
+			FOperation::ADDRESS VarAddress;
+			FMemory::Memcpy(&VarAddress, Data, sizeof(FOperation::ADDRESS));
+			Data += sizeof(FOperation::ADDRESS);
 
-			OP::ADDRESS DefAddress;
-			FMemory::Memcpy(&DefAddress, Data, sizeof(OP::ADDRESS));
-			Data += sizeof(OP::ADDRESS);
+			FOperation::ADDRESS DefAddress;
+			FMemory::Memcpy(&DefAddress, Data, sizeof(FOperation::ADDRESS));
+			Data += sizeof(FOperation::ADDRESS);
 
-			OP::FSwitchCaseDescriptor CaseDesc;
-			FMemory::Memcpy(&CaseDesc, Data, sizeof(OP::FSwitchCaseDescriptor));
-			Data += sizeof(OP::FSwitchCaseDescriptor);
+			FOperation::FSwitchCaseDescriptor CaseDesc;
+			FMemory::Memcpy(&CaseDesc, Data, sizeof(FOperation::FSwitchCaseDescriptor));
+			Data += sizeof(FOperation::FSwitchCaseDescriptor);
 
 			f(VarAddress);
 			f(DefAddress);
@@ -270,9 +105,9 @@ namespace UE::Mutable::Private
 					//FMemory::Memcpy( &Condition, data, sizeof(int32));
 					Data += sizeof(int32);
 
-					OP::ADDRESS CaseAt;
-					FMemory::Memcpy(&CaseAt, Data, sizeof(OP::ADDRESS));
-					Data += sizeof(OP::ADDRESS);
+					FOperation::ADDRESS CaseAt;
+					FMemory::Memcpy(&CaseAt, Data, sizeof(FOperation::ADDRESS));
+					Data += sizeof(FOperation::ADDRESS);
 
 					f(CaseAt);
 				}
@@ -289,9 +124,9 @@ namespace UE::Mutable::Private
 					//FMemory::Memcpy( &RangeSize, data, sizeof(uint32));
 					Data += sizeof(int32);
 
-					OP::ADDRESS CaseAt;
-					FMemory::Memcpy(&CaseAt, Data, sizeof(OP::ADDRESS));
-					Data += sizeof(OP::ADDRESS);
+					FOperation::ADDRESS CaseAt;
+					FMemory::Memcpy(&CaseAt, Data, sizeof(FOperation::ADDRESS));
+					Data += sizeof(FOperation::ADDRESS);
 
 					f(CaseAt);
 				}
@@ -304,19 +139,19 @@ namespace UE::Mutable::Private
 		case EOpType::CO_MATERIAL_BREAK:
 		case EOpType::IM_MATERIAL_BREAK:
 		{
-			OP::MaterialBreakArgs Args = program.GetOpArgs<OP::MaterialBreakArgs>(at);
+			FOperation::MaterialBreakArgs Args = program.GetOpArgs<FOperation::MaterialBreakArgs>(at);
 			f(Args.Material);
 			break;
 		}
 		
 		case EOpType::MI_MODIFY:
 		{
-			OP::MaterialModifyArgs Args = program.GetOpArgs<OP::MaterialModifyArgs>(at);
+			FOperation::MaterialModifyArgs Args = program.GetOpArgs<FOperation::MaterialModifyArgs>(at);
 			int32 ParameterCount = Args.NumParameters;
 
 			//Dynamic data can't be stored into arg structures
 			const uint8* Data = program.GetOpArgsPointer(at);
-			Data += sizeof(OP::MaterialModifyArgs);
+			Data += sizeof(FOperation::MaterialModifyArgs);
 
 			// First, the material
 			f(Args.Material);
@@ -324,18 +159,18 @@ namespace UE::Mutable::Private
 			// Store parameter names and schedule the parameter operations
 			for (int32 ParameterIndex = 0; ParameterIndex < ParameterCount; ++ParameterIndex)
 			{
-				OP::ADDRESS NewName;
-				FMemory::Memcpy(&NewName, Data, sizeof(OP::ADDRESS));
-				Data += sizeof(OP::ADDRESS);
+				FOperation::ADDRESS NewName;
+				FMemory::Memcpy(&NewName, Data, sizeof(FOperation::ADDRESS));
+				Data += sizeof(FOperation::ADDRESS);
 
 				int8 LayerIndex;
 				FMemory::Memcpy(&LayerIndex, Data, sizeof(int8));
 				Data += sizeof(int8);
 
-				OP::ADDRESS NewParameterOperation;
-				FMemory::Memcpy(&NewParameterOperation, Data, sizeof(OP::ADDRESS));
+				FOperation::ADDRESS NewParameterOperation;
+				FMemory::Memcpy(&NewParameterOperation, Data, sizeof(FOperation::ADDRESS));
 				f(NewParameterOperation);
-				Data += sizeof(OP::ADDRESS);
+				Data += sizeof(FOperation::ADDRESS);
 
 				// Evaluate non lazy branches here
 				if (GetOpDataType(program.GetOpType(NewParameterOperation)) == EDataType::Image)
@@ -351,7 +186,7 @@ namespace UE::Mutable::Private
         //-------------------------------------------------------------------------------------
         case EOpType::BO_EQUAL_INT_CONST:
         {
-			OP::BoolEqualScalarConstArgs args = program.GetOpArgs<OP::BoolEqualScalarConstArgs>(at);
+			FOperation::BoolEqualScalarConstArgs args = program.GetOpArgs<FOperation::BoolEqualScalarConstArgs>(at);
             f(args.Value );
             break;
         }
@@ -359,7 +194,7 @@ namespace UE::Mutable::Private
         case EOpType::BO_AND:
         case EOpType::BO_OR:
         {
-			OP::BoolBinaryArgs args = program.GetOpArgs<OP::BoolBinaryArgs>(at);
+			FOperation::BoolBinaryArgs args = program.GetOpArgs<FOperation::BoolBinaryArgs>(at);
             f(args.A );
             f(args.B );
             break;
@@ -367,14 +202,14 @@ namespace UE::Mutable::Private
 
         case EOpType::BO_NOT:
         {
-			OP::BoolNotArgs args = program.GetOpArgs<OP::BoolNotArgs>(at);
+			FOperation::BoolNotArgs args = program.GetOpArgs<FOperation::BoolNotArgs>(at);
             f(args.A );
             break;
         }
 
         case EOpType::SC_ARITHMETIC:
         {
-			OP::ArithmeticArgs args = program.GetOpArgs<OP::ArithmeticArgs>(at);
+			FOperation::ArithmeticArgs args = program.GetOpArgs<FOperation::ArithmeticArgs>(at);
             f(args.A );
             f(args.B );
             break;
@@ -382,7 +217,7 @@ namespace UE::Mutable::Private
 
         case EOpType::CO_SAMPLEIMAGE:
         {
-			OP::ColourSampleImageArgs args = program.GetOpArgs<OP::ColourSampleImageArgs>(at);
+			FOperation::ColourSampleImageArgs args = program.GetOpArgs<FOperation::ColourSampleImageArgs>(at);
             f(args.Image );
             f(args.X );
             f(args.Y );
@@ -391,7 +226,7 @@ namespace UE::Mutable::Private
 
         case EOpType::CO_SWIZZLE:
         {
-			OP::ColourSwizzleArgs args = program.GetOpArgs<OP::ColourSwizzleArgs>(at);
+			FOperation::ColourSwizzleArgs args = program.GetOpArgs<FOperation::ColourSwizzleArgs>(at);
             for (int t=0;t<MUTABLE_OP_MAX_SWIZZLE_CHANNELS;++t)
             {
                 f(args.sources[t] );
@@ -401,7 +236,7 @@ namespace UE::Mutable::Private
 
         case EOpType::CO_FROMSCALARS:
         {
-			OP::ColourFromScalarsArgs args = program.GetOpArgs<OP::ColourFromScalarsArgs>(at);
+			FOperation::ColourFromScalarsArgs args = program.GetOpArgs<FOperation::ColourFromScalarsArgs>(at);
 			for (int32 t = 0; t < MUTABLE_OP_MAX_SWIZZLE_CHANNELS; ++t)
 			{
 				f(args.V[t]);
@@ -411,7 +246,7 @@ namespace UE::Mutable::Private
 
         case EOpType::CO_ARITHMETIC:
         {
-			OP::ArithmeticArgs args = program.GetOpArgs<OP::ArithmeticArgs>(at);
+			FOperation::ArithmeticArgs args = program.GetOpArgs<FOperation::ArithmeticArgs>(at);
             f(args.A);
             f(args.B);
             break;
@@ -419,7 +254,7 @@ namespace UE::Mutable::Private
 		
 		case EOpType::CO_LINEARTOSRGB:
         {
-			OP::ColorArgs Args = program.GetOpArgs<OP::ColorArgs>(at);
+			FOperation::ColorArgs Args = program.GetOpArgs<FOperation::ColorArgs>(at);
             f(Args.Color);
             break;
         }
@@ -427,7 +262,7 @@ namespace UE::Mutable::Private
         //-------------------------------------------------------------------------------------
         case EOpType::IM_LAYER:
         {
-			OP::ImageLayerArgs args = program.GetOpArgs<OP::ImageLayerArgs>(at);
+			FOperation::ImageLayerArgs args = program.GetOpArgs<FOperation::ImageLayerArgs>(at);
             f(args.base );
         	f(args.mask );
             f(args.blended );
@@ -436,7 +271,7 @@ namespace UE::Mutable::Private
 
         case EOpType::IM_LAYERCOLOUR:
         {
-			OP::ImageLayerColourArgs args = program.GetOpArgs<OP::ImageLayerColourArgs>(at);
+			FOperation::ImageLayerColourArgs args = program.GetOpArgs<FOperation::ImageLayerColourArgs>(at);
             f(args.base );
         	f(args.mask );
             f(args.colour );
@@ -445,7 +280,7 @@ namespace UE::Mutable::Private
 
         case EOpType::IM_MULTILAYER:
         {
-			OP::ImageMultiLayerArgs args = program.GetOpArgs<OP::ImageMultiLayerArgs>(at);
+			FOperation::ImageMultiLayerArgs args = program.GetOpArgs<FOperation::ImageMultiLayerArgs>(at);
             f(args.rangeSize );
             f(args.base );
             f(args.mask );
@@ -455,7 +290,7 @@ namespace UE::Mutable::Private
 
 		case EOpType::IM_NORMALCOMPOSITE:
 		{
-			OP::ImageNormalCompositeArgs args = program.GetOpArgs<OP::ImageNormalCompositeArgs>(at);
+			FOperation::ImageNormalCompositeArgs args = program.GetOpArgs<FOperation::ImageNormalCompositeArgs>(at);
 			f(args.base);
 			f(args.normal);
 
@@ -464,28 +299,28 @@ namespace UE::Mutable::Private
 
         case EOpType::IM_PIXELFORMAT:
         {
-			OP::ImagePixelFormatArgs args = program.GetOpArgs<OP::ImagePixelFormatArgs>(at);
+			FOperation::ImagePixelFormatArgs args = program.GetOpArgs<FOperation::ImagePixelFormatArgs>(at);
             f(args.source );
             break;
         }
 
         case EOpType::IM_MIPMAP:
         {
-			OP::ImageMipmapArgs args = program.GetOpArgs<OP::ImageMipmapArgs>(at);
+			FOperation::ImageMipmapArgs args = program.GetOpArgs<FOperation::ImageMipmapArgs>(at);
             f(args.source );
             break;
         }
 
         case EOpType::IM_RESIZE:
         {
-			OP::ImageResizeArgs Args = program.GetOpArgs<OP::ImageResizeArgs>(at);
+			FOperation::ImageResizeArgs Args = program.GetOpArgs<FOperation::ImageResizeArgs>(at);
             f(Args.Source );
             break;
         }
 
         case EOpType::IM_RESIZELIKE:
         {
-			OP::ImageResizeLikeArgs Args = program.GetOpArgs<OP::ImageResizeLikeArgs>(at);
+			FOperation::ImageResizeLikeArgs Args = program.GetOpArgs<FOperation::ImageResizeLikeArgs>(at);
             f(Args.Source );
             f(Args.SizeSource );
             break;
@@ -493,21 +328,21 @@ namespace UE::Mutable::Private
 
         case EOpType::IM_RESIZEREL:
         {
-			OP::ImageResizeRelArgs Args = program.GetOpArgs<OP::ImageResizeRelArgs>(at);
+			FOperation::ImageResizeRelArgs Args = program.GetOpArgs<FOperation::ImageResizeRelArgs>(at);
             f(Args.Source );
             break;
         }
 
         case EOpType::IM_BLANKLAYOUT:
         {
-			OP::ImageBlankLayoutArgs Args = program.GetOpArgs<OP::ImageBlankLayoutArgs>(at);
+			FOperation::ImageBlankLayoutArgs Args = program.GetOpArgs<FOperation::ImageBlankLayoutArgs>(at);
             f(Args.Layout );
             break;
         }
 
         case EOpType::IM_COMPOSE:
         {
-			OP::ImageComposeArgs args = program.GetOpArgs<OP::ImageComposeArgs>(at);
+			FOperation::ImageComposeArgs args = program.GetOpArgs<FOperation::ImageComposeArgs>(at);
             f(args.layout );
             f(args.base );
             f(args.blockImage );
@@ -517,7 +352,7 @@ namespace UE::Mutable::Private
 
 		case EOpType::IM_MULTICOMPOSE:
 		{
-			OP::ImageMultiComposeArgs args = program.GetOpArgs<OP::ImageMultiComposeArgs>(at);
+			FOperation::ImageMultiComposeArgs args = program.GetOpArgs<FOperation::ImageMultiComposeArgs>(at);
 			f(args.Layout);
 			f(args.Base);
 			f(args.SourceLayout);
@@ -527,7 +362,7 @@ namespace UE::Mutable::Private
 
         case EOpType::IM_INTERPOLATE:
         {
-			OP::ImageInterpolateArgs Args = program.GetOpArgs<OP::ImageInterpolateArgs>(at);
+			FOperation::ImageInterpolateArgs Args = program.GetOpArgs<FOperation::ImageInterpolateArgs>(at);
             f(Args.Factor );
 
             for (int32 TargetIndex=0; TargetIndex <MUTABLE_OP_MAX_INTERPOLATE_COUNT;++TargetIndex)
@@ -539,7 +374,7 @@ namespace UE::Mutable::Private
 
         case EOpType::IM_SWIZZLE:
         {
-			OP::ImageSwizzleArgs args = program.GetOpArgs<OP::ImageSwizzleArgs>(at);
+			FOperation::ImageSwizzleArgs args = program.GetOpArgs<FOperation::ImageSwizzleArgs>(at);
             for (int t=0;t<MUTABLE_OP_MAX_SWIZZLE_CHANNELS;++t)
             {
                 f(args.sources[t] );
@@ -549,7 +384,7 @@ namespace UE::Mutable::Private
 
         case EOpType::IM_SATURATE:
         {
-			OP::ImageSaturateArgs Args = program.GetOpArgs<OP::ImageSaturateArgs>(at);
+			FOperation::ImageSaturateArgs Args = program.GetOpArgs<FOperation::ImageSaturateArgs>(at);
             f(Args.Base );
             f(Args.Factor );
             break;
@@ -557,14 +392,14 @@ namespace UE::Mutable::Private
 
         case EOpType::IM_LUMINANCE:
         {
-			OP::ImageLuminanceArgs Args = program.GetOpArgs<OP::ImageLuminanceArgs>(at);
+			FOperation::ImageLuminanceArgs Args = program.GetOpArgs<FOperation::ImageLuminanceArgs>(at);
             f(Args.Base );
             break;
         }
 
         case EOpType::IM_COLOURMAP:
         {
-			OP::ImageColourMapArgs Args = program.GetOpArgs<OP::ImageColourMapArgs>(at);
+			FOperation::ImageColourMapArgs Args = program.GetOpArgs<FOperation::ImageColourMapArgs>(at);
             f(Args.Base );
             f(Args.Mask );
             f(Args.Map );
@@ -573,7 +408,7 @@ namespace UE::Mutable::Private
 
         case EOpType::IM_BINARISE:
         {
-			OP::ImageBinariseArgs Args = program.GetOpArgs<OP::ImageBinariseArgs>(at);
+			FOperation::ImageBinariseArgs Args = program.GetOpArgs<FOperation::ImageBinariseArgs>(at);
             f(Args.Base );
             f(Args.Threshold );
             break;
@@ -581,21 +416,21 @@ namespace UE::Mutable::Private
 
         case EOpType::IM_PLAINCOLOUR:
         {
-			OP::ImagePlainColorArgs args = program.GetOpArgs<OP::ImagePlainColorArgs>(at);
+			FOperation::ImagePlainColorArgs args = program.GetOpArgs<FOperation::ImagePlainColorArgs>(at);
             f(args.Color );
             break;
         }
 
         case EOpType::IM_CROP:
         {
-			OP::ImageCropArgs args = program.GetOpArgs<OP::ImageCropArgs>(at);
+			FOperation::ImageCropArgs args = program.GetOpArgs<FOperation::ImageCropArgs>(at);
             f(args.source );
             break;
         }
 
         case EOpType::IM_PATCH:
         {
-			OP::ImagePatchArgs args = program.GetOpArgs<OP::ImagePatchArgs>(at);
+			FOperation::ImagePatchArgs args = program.GetOpArgs<FOperation::ImagePatchArgs>(at);
             f(args.base );
             f(args.patch );
             break;
@@ -603,7 +438,7 @@ namespace UE::Mutable::Private
 
         case EOpType::IM_RASTERMESH:
         {
-			OP::ImageRasterMeshArgs args = program.GetOpArgs<OP::ImageRasterMeshArgs>(at);
+			FOperation::ImageRasterMeshArgs args = program.GetOpArgs<FOperation::ImageRasterMeshArgs>(at);
             f(args.mesh );
             f(args.image );
             f(args.mask );
@@ -614,14 +449,14 @@ namespace UE::Mutable::Private
 
         case EOpType::IM_MAKEGROWMAP:
         {
-			OP::ImageMakeGrowMapArgs args = program.GetOpArgs<OP::ImageMakeGrowMapArgs>(at);
+			FOperation::ImageMakeGrowMapArgs args = program.GetOpArgs<FOperation::ImageMakeGrowMapArgs>(at);
             f(args.mask );
             break;
         }
 
         case EOpType::IM_DISPLACE:
         {
-			OP::ImageDisplaceArgs Args = program.GetOpArgs<OP::ImageDisplaceArgs>(at);
+			FOperation::ImageDisplaceArgs Args = program.GetOpArgs<FOperation::ImageDisplaceArgs>(at);
             f(Args.Source );
             f(Args.DisplacementMap );
             break;
@@ -629,14 +464,14 @@ namespace UE::Mutable::Private
 
 		case EOpType::IM_INVERT:
 		{
-			OP::ImageInvertArgs args = program.GetOpArgs<OP::ImageInvertArgs>(at);
+			FOperation::ImageInvertArgs args = program.GetOpArgs<FOperation::ImageInvertArgs>(at);
 			f(args.Base);
 			break;
 		}
 
 		case EOpType::IM_TRANSFORM:
 		{
-			OP::ImageTransformArgs Args = program.GetOpArgs<OP::ImageTransformArgs>(at);
+			FOperation::ImageTransformArgs Args = program.GetOpArgs<FOperation::ImageTransformArgs>(at);
 			f(Args.Base);
 			f(Args.OffsetX);
 			f(Args.OffsetY);
@@ -648,7 +483,7 @@ namespace UE::Mutable::Private
 
         case EOpType::IM_PARAMETER_CONVERT:
         {
-        	OP::ImageParameterConvertArgs Args = program.GetOpArgs<OP::ImageParameterConvertArgs>(at);
+        	FOperation::ImageParameterConvertArgs Args = program.GetOpArgs<FOperation::ImageParameterConvertArgs>(at);
         	f(Args.ImageParameter);
         	break;
         }
@@ -656,7 +491,7 @@ namespace UE::Mutable::Private
         //-------------------------------------------------------------------------------------
         case EOpType::ME_APPLYLAYOUT:
         {
-			OP::MeshApplyLayoutArgs args = program.GetOpArgs<OP::MeshApplyLayoutArgs>(at);
+			FOperation::MeshApplyLayoutArgs args = program.GetOpArgs<FOperation::MeshApplyLayoutArgs>(at);
             f(args.Layout );
             f(args.Mesh );
             break;
@@ -664,7 +499,7 @@ namespace UE::Mutable::Private
 
 		case EOpType::ME_PREPARELAYOUT:
 		{
-			OP::MeshPrepareLayoutArgs Args = program.GetOpArgs<OP::MeshPrepareLayoutArgs>(at);
+			FOperation::MeshPrepareLayoutArgs Args = program.GetOpArgs<FOperation::MeshPrepareLayoutArgs>(at);
 			f(Args.Layout);
 			f(Args.Mesh);
 			break;
@@ -674,19 +509,19 @@ namespace UE::Mutable::Private
         {
 			const uint8_t* data = program.GetOpArgsPointer(at);
 
-			OP::ADDRESS BaseAt = 0;
-			FMemory::Memcpy(&BaseAt, data, sizeof(OP::ADDRESS)); data += sizeof(OP::ADDRESS);
+			FOperation::ADDRESS BaseAt = 0;
+			FMemory::Memcpy(&BaseAt, data, sizeof(FOperation::ADDRESS)); data += sizeof(FOperation::ADDRESS);
 			f(BaseAt);
 
-			OP::ADDRESS TargetAt = 0;
-			FMemory::Memcpy(&TargetAt, data, sizeof(OP::ADDRESS)); data += sizeof(OP::ADDRESS);
+			FOperation::ADDRESS TargetAt = 0;
+			FMemory::Memcpy(&TargetAt, data, sizeof(FOperation::ADDRESS)); data += sizeof(FOperation::ADDRESS);
 			f(TargetAt);
 			break;
         }
 
         case EOpType::ME_MORPH:
         {
-			OP::MeshMorphArgs Args = program.GetOpArgs<OP::MeshMorphArgs>(at);
+			FOperation::MeshMorphArgs Args = program.GetOpArgs<FOperation::MeshMorphArgs>(at);
 
 			f(Args.Factor);
 			f(Args.Base);
@@ -696,7 +531,7 @@ namespace UE::Mutable::Private
 
         case EOpType::ME_MERGE:
         {
-			OP::MeshMergeArgs args = program.GetOpArgs<OP::MeshMergeArgs>(at);
+			FOperation::MeshMergeArgs args = program.GetOpArgs<FOperation::MeshMergeArgs>(at);
             f(args.Base );
             f(args.Added );
             break;
@@ -704,7 +539,7 @@ namespace UE::Mutable::Private
 
 		case EOpType::ME_MASKCLIPMESH:
 		{
-			OP::MeshMaskClipMeshArgs args = program.GetOpArgs<OP::MeshMaskClipMeshArgs>(at);
+			FOperation::MeshMaskClipMeshArgs args = program.GetOpArgs<FOperation::MeshMaskClipMeshArgs>(at);
 			f(args.source);
 			f(args.clip);
 			break;
@@ -712,7 +547,7 @@ namespace UE::Mutable::Private
 
 		case EOpType::ME_MASKCLIPUVMASK:
 		{
-			OP::MeshMaskClipUVMaskArgs Args = program.GetOpArgs<OP::MeshMaskClipUVMaskArgs>(at);
+			FOperation::MeshMaskClipUVMaskArgs Args = program.GetOpArgs<FOperation::MeshMaskClipUVMaskArgs>(at);
 			f(Args.Source);
 			f(Args.UVSource);
 			f(Args.MaskImage);
@@ -722,7 +557,7 @@ namespace UE::Mutable::Private
 
         case EOpType::ME_MASKDIFF:
         {
-			OP::MeshMaskDiffArgs args = program.GetOpArgs<OP::MeshMaskDiffArgs>(at);
+			FOperation::MeshMaskDiffArgs args = program.GetOpArgs<FOperation::MeshMaskDiffArgs>(at);
             f(args.Source );
             f(args.Fragment );
             break;
@@ -731,8 +566,8 @@ namespace UE::Mutable::Private
         case EOpType::ME_REMOVEMASK:
         {
             const uint8* data = program.GetOpArgsPointer(at);
-            UE::Mutable::Private::OP::ADDRESS source;
-            FMemory::Memcpy( &source, data, sizeof(OP::ADDRESS) ); data+=sizeof(OP::ADDRESS);
+            UE::Mutable::Private::FOperation::ADDRESS source;
+            FMemory::Memcpy( &source, data, sizeof(FOperation::ADDRESS) ); data+=sizeof(FOperation::ADDRESS);
             f(source);
 
 			EFaceCullStrategy FaceCullStrategy;
@@ -743,12 +578,12 @@ namespace UE::Mutable::Private
 			FMemory::Memcpy( &removes, data, sizeof(uint16) ); data+=sizeof(uint16);
             for (uint16 r=0; r<removes; ++r)
             {
-                UE::Mutable::Private::OP::ADDRESS condition;
-				FMemory::Memcpy( &condition, data, sizeof(OP::ADDRESS) ); data+=sizeof(OP::ADDRESS);
+                UE::Mutable::Private::FOperation::ADDRESS condition;
+				FMemory::Memcpy( &condition, data, sizeof(FOperation::ADDRESS) ); data+=sizeof(FOperation::ADDRESS);
                 f(condition);
 
-                UE::Mutable::Private::OP::ADDRESS mask;
-				FMemory::Memcpy( &mask, data, sizeof(OP::ADDRESS) ); data+=sizeof(OP::ADDRESS);
+                UE::Mutable::Private::FOperation::ADDRESS mask;
+				FMemory::Memcpy( &mask, data, sizeof(FOperation::ADDRESS) ); data+=sizeof(FOperation::ADDRESS);
                 f(mask);
             }
             break;
@@ -756,14 +591,14 @@ namespace UE::Mutable::Private
 
 		case EOpType::ME_ADDMETADATA:
 		{
-			OP::MeshAddMetadataArgs Args = program.GetOpArgs<OP::MeshAddMetadataArgs>(at);
+			FOperation::MeshAddMetadataArgs Args = program.GetOpArgs<FOperation::MeshAddMetadataArgs>(at);
 			f(Args.Source);
 			break;
 		}
 
         case EOpType::ME_FORMAT:
         {
-			OP::MeshFormatArgs args = program.GetOpArgs<OP::MeshFormatArgs>(at);
+			FOperation::MeshFormatArgs args = program.GetOpArgs<FOperation::MeshFormatArgs>(at);
             f(args.source );
             f(args.format );
             break;
@@ -771,7 +606,7 @@ namespace UE::Mutable::Private
 
         case EOpType::ME_TRANSFORM:
         {
-			OP::MeshTransformArgs args = program.GetOpArgs<OP::MeshTransformArgs>(at);
+			FOperation::MeshTransformArgs args = program.GetOpArgs<FOperation::MeshTransformArgs>(at);
             f(args.source );
             break;
         }
@@ -779,22 +614,22 @@ namespace UE::Mutable::Private
         case EOpType::ME_EXTRACTLAYOUTBLOCK:
         {
             const uint8_t* data = program.GetOpArgsPointer(at);
-            UE::Mutable::Private::OP::ADDRESS source;
-			FMemory::Memcpy( &source, data, sizeof(OP::ADDRESS) );
+            UE::Mutable::Private::FOperation::ADDRESS source;
+			FMemory::Memcpy( &source, data, sizeof(FOperation::ADDRESS) );
             f(source);
             break;
         }
 
         case EOpType::ME_CLIPMORPHPLANE:
         {
-			OP::MeshClipMorphPlaneArgs args = program.GetOpArgs<OP::MeshClipMorphPlaneArgs>(at);
+			FOperation::MeshClipMorphPlaneArgs args = program.GetOpArgs<FOperation::MeshClipMorphPlaneArgs>(at);
             f(args.Source);
             break;
         }
 
         case EOpType::ME_CLIPWITHMESH :
         {
-			OP::MeshClipWithMeshArgs Args = program.GetOpArgs<OP::MeshClipWithMeshArgs>(at);
+			FOperation::MeshClipWithMeshArgs Args = program.GetOpArgs<FOperation::MeshClipWithMeshArgs>(at);
             f(Args.Source);
             f(Args.ClipMesh);
             break;
@@ -802,7 +637,7 @@ namespace UE::Mutable::Private
 
         case EOpType::ME_CLIPDEFORM:
         {
-			OP::MeshClipDeformArgs args = program.GetOpArgs<OP::MeshClipDeformArgs>(at);
+			FOperation::MeshClipDeformArgs args = program.GetOpArgs<FOperation::MeshClipDeformArgs>(at);
             f(args.mesh);
             f(args.clipShape);
             break;
@@ -810,7 +645,7 @@ namespace UE::Mutable::Private
        
 		case EOpType::ME_MORPHRESHAPE:
 		{
-			OP::MeshMorphReshapeArgs Args = program.GetOpArgs<OP::MeshMorphReshapeArgs>(at);
+			FOperation::MeshMorphReshapeArgs Args = program.GetOpArgs<FOperation::MeshMorphReshapeArgs>(at);
 			f(Args.Morph);
 			f(Args.Reshape);
 
@@ -819,7 +654,7 @@ namespace UE::Mutable::Private
 
         case EOpType::ME_PROJECT :
         {
-			OP::MeshProjectArgs Args = program.GetOpArgs<OP::MeshProjectArgs>(at);
+			FOperation::MeshProjectArgs Args = program.GetOpArgs<FOperation::MeshProjectArgs>(at);
             f(Args.Mesh);
             f(Args.Projector);
             break;
@@ -827,7 +662,7 @@ namespace UE::Mutable::Private
 
 		case EOpType::ME_APPLYPOSE:
 		{
-			OP::MeshApplyPoseArgs Args = program.GetOpArgs<OP::MeshApplyPoseArgs>(at);
+			FOperation::MeshApplyPoseArgs Args = program.GetOpArgs<FOperation::MeshApplyPoseArgs>(at);
 			f(Args.base);
 			f(Args.pose);
 			break;
@@ -835,7 +670,7 @@ namespace UE::Mutable::Private
 
 		case EOpType::ME_BINDSHAPE:
 		{
-			OP::MeshBindShapeArgs args = program.GetOpArgs<OP::MeshBindShapeArgs>(at);
+			FOperation::MeshBindShapeArgs args = program.GetOpArgs<FOperation::MeshBindShapeArgs>(at);
 			f(args.mesh);
 			f(args.shape);
 			break;
@@ -843,7 +678,7 @@ namespace UE::Mutable::Private
 
 		case EOpType::ME_APPLYSHAPE:
 		{
-			OP::MeshApplyShapeArgs args = program.GetOpArgs<OP::MeshApplyShapeArgs>(at);
+			FOperation::MeshApplyShapeArgs args = program.GetOpArgs<FOperation::MeshApplyShapeArgs>(at);
 			f(args.mesh);
 			f(args.shape);
 			break;
@@ -851,14 +686,14 @@ namespace UE::Mutable::Private
 
 		case EOpType::ME_OPTIMIZESKINNING:
 		{
-			OP::MeshOptimizeSkinningArgs args = program.GetOpArgs<OP::MeshOptimizeSkinningArgs>(at);
+			FOperation::MeshOptimizeSkinningArgs args = program.GetOpArgs<FOperation::MeshOptimizeSkinningArgs>(at);
 			f(args.source);
 			break;
 		}
 
         case EOpType::ME_TRANSFORMWITHMESH :
         {
-        	OP::MeshTransformWithinMeshArgs args = program.GetOpArgs<OP::MeshTransformWithinMeshArgs>(at);
+        	FOperation::MeshTransformWithinMeshArgs args = program.GetOpArgs<FOperation::MeshTransformWithinMeshArgs>(at);
         	f(args.sourceMesh);
         	f(args.boundingMesh);
         	f(args.matrix);
@@ -867,7 +702,7 @@ namespace UE::Mutable::Private
 
 		case EOpType::ME_TRANSFORMWITHBONE:
 		{
-			OP::MeshTransformWithBoneArgs args = program.GetOpArgs<OP::MeshTransformWithBoneArgs>(at);
+			FOperation::MeshTransformWithBoneArgs args = program.GetOpArgs<FOperation::MeshTransformWithBoneArgs>(at);
 			f(args.SourceMesh);
 			f(args.Matrix);
 			break;
@@ -875,7 +710,7 @@ namespace UE::Mutable::Private
 
         case EOpType::ME_SKELETALMESH_BREAK:
         {
-        	OP::MeshSkeletalMeshBreakArgs args = program.GetOpArgs<OP::MeshSkeletalMeshBreakArgs>(at);
+        	FOperation::MeshSkeletalMeshBreakArgs args = program.GetOpArgs<FOperation::MeshSkeletalMeshBreakArgs>(at);
         	f(args.SkeletalMeshParameter);
         	break;
         }
@@ -885,7 +720,7 @@ namespace UE::Mutable::Private
         case EOpType::IN_ADDSTRING:
         case EOpType::IN_ADDSURFACE:
         {
-			OP::InstanceAddArgs args = program.GetOpArgs<OP::InstanceAddArgs>(at);
+			FOperation::InstanceAddArgs args = program.GetOpArgs<FOperation::InstanceAddArgs>(at);
             f(args.instance );
             f(args.value );
             break;
@@ -893,7 +728,7 @@ namespace UE::Mutable::Private
 
         case EOpType::IN_ADDCOMPONENT:
         {
-        	OP::InstanceAddComponentArgs args = program.GetOpArgs<OP::InstanceAddComponentArgs>(at);
+        	FOperation::InstanceAddComponentArgs args = program.GetOpArgs<FOperation::InstanceAddComponentArgs>(at);
         	f(args.Instance );
         	f(args.Value );
         	break;
@@ -909,9 +744,9 @@ namespace UE::Mutable::Private
 
 			for (int8 LODIndex = 0; LODIndex < LODCount; ++LODIndex)
 			{
-				OP::ADDRESS LODAddress;
-				FMemory::Memcpy(&LODAddress, Data, sizeof(OP::ADDRESS));
-				Data += sizeof(OP::ADDRESS);
+				FOperation::ADDRESS LODAddress;
+				FMemory::Memcpy(&LODAddress, Data, sizeof(FOperation::ADDRESS));
+				Data += sizeof(FOperation::ADDRESS);
 
 				f(LODAddress);
             }
@@ -920,7 +755,7 @@ namespace UE::Mutable::Private
 
 		case EOpType::IN_ADDEXTENSIONDATA:
 		{
-			const OP::InstanceAddExtensionDataArgs Args = program.GetOpArgs<OP::InstanceAddExtensionDataArgs>(at);
+			const FOperation::InstanceAddExtensionDataArgs Args = program.GetOpArgs<FOperation::InstanceAddExtensionDataArgs>(at);
 
 			f(Args.Instance);
 			f(Args.ExtensionData);
@@ -930,7 +765,7 @@ namespace UE::Mutable::Private
 
 		case EOpType::IN_ADDOVERRIDEMATERIAL:
 		{
-        		const OP::InstanceAddOverrideMaterialArgs Args = program.GetOpArgs<OP::InstanceAddOverrideMaterialArgs>(at);
+        		const FOperation::InstanceAddOverrideMaterialArgs Args = program.GetOpArgs<FOperation::InstanceAddOverrideMaterialArgs>(at);
 
         		f(Args.Instance);
         		f(Args.Material);
@@ -940,7 +775,7 @@ namespace UE::Mutable::Private
 
 		case EOpType::IN_ADDOVERLAYMATERIAL:
         {
-        	const OP::InstanceAddOverlayMaterialArgs Args = program.GetOpArgs<OP::InstanceAddOverlayMaterialArgs>(at);
+        	const FOperation::InstanceAddOverlayMaterialArgs Args = program.GetOpArgs<FOperation::InstanceAddOverlayMaterialArgs>(at);
 
         	f(Args.Instance);
         	f(Args.Material);
@@ -950,7 +785,7 @@ namespace UE::Mutable::Private
         	
 		case EOpType::IN_ADDMATERIAL:
 		{
-			const OP::InstanceAddMaterialArgs Args = program.GetOpArgs<OP::InstanceAddMaterialArgs>(at);
+			const FOperation::InstanceAddMaterialArgs Args = program.GetOpArgs<FOperation::InstanceAddMaterialArgs>(at);
 
 			f(Args.Instance);
 			f(Args.Material);
@@ -960,7 +795,7 @@ namespace UE::Mutable::Private
 
         case EOpType::IN_ADDSKELETALMESH:
         {
-        	const OP::InstanceAddArgs Args = program.GetOpArgs<OP::InstanceAddArgs>(at);
+        	const FOperation::InstanceAddArgs Args = program.GetOpArgs<FOperation::InstanceAddArgs>(at);
 
         	f(Args.instance);
         	f(Args.value);
@@ -971,14 +806,14 @@ namespace UE::Mutable::Private
         //-------------------------------------------------------------------------------------
         case EOpType::LA_PACK:
         {
-			OP::LayoutPackArgs args = program.GetOpArgs<OP::LayoutPackArgs>(at);
+			FOperation::LayoutPackArgs args = program.GetOpArgs<FOperation::LayoutPackArgs>(at);
             f(args.Source );
             break;
         }
 
         case EOpType::LA_MERGE:
         {
-			OP::LayoutMergeArgs args = program.GetOpArgs<OP::LayoutMergeArgs>(at);
+			FOperation::LayoutMergeArgs args = program.GetOpArgs<FOperation::LayoutMergeArgs>(at);
             f(args.Base );
             f(args.Added );
             break;
@@ -986,7 +821,7 @@ namespace UE::Mutable::Private
 
 		case EOpType::LA_REMOVEBLOCKS:
 		{
-			OP::LayoutRemoveBlocksArgs args = program.GetOpArgs<OP::LayoutRemoveBlocksArgs>(at);
+			FOperation::LayoutRemoveBlocksArgs args = program.GetOpArgs<FOperation::LayoutRemoveBlocksArgs>(at);
 			f(args.Source);
 			f(args.ReferenceLayout);
 			break;
@@ -994,7 +829,7 @@ namespace UE::Mutable::Private
 
 		case EOpType::LA_FROMMESH:
 		{
-			OP::LayoutFromMeshArgs args = program.GetOpArgs<OP::LayoutFromMeshArgs>(at);
+			FOperation::LayoutFromMeshArgs args = program.GetOpArgs<FOperation::LayoutFromMeshArgs>(at);
 			f(args.Mesh);
 			break;
 		}
@@ -1008,14 +843,14 @@ namespace UE::Mutable::Private
         {
         	const uint8* Data = program.GetOpArgsPointer(at);
 
-        	OP::ExternalArgs Args = program.GetOpArgs<OP::ExternalArgs>(at);
-        	Data += sizeof(OP::ExternalArgs);
+        	FOperation::ExternalArgs Args = program.GetOpArgs<FOperation::ExternalArgs>(at);
+        	Data += sizeof(FOperation::ExternalArgs);
         		
         	for (int32 Index = 0; Index < Args.NumOperants; ++Index)
         	{
-        		OP::ADDRESS OperantAt;
-        		FMemory::Memcpy(&OperantAt, Data, sizeof(OP::ADDRESS));
-        		Data += sizeof(OP::ADDRESS);
+        		FOperation::ADDRESS OperantAt;
+        		FMemory::Memcpy(&OperantAt, Data, sizeof(FOperation::ADDRESS));
+        		Data += sizeof(FOperation::ADDRESS);
 
         		f(OperantAt);
         	}
@@ -1025,7 +860,7 @@ namespace UE::Mutable::Private
 
         case EOpType::MI_SKELETALMESH_BREAK:
         {
-        	const OP::MaterialSkeletalMeshBreakArgs Args = program.GetOpArgs<OP::MaterialSkeletalMeshBreakArgs>(at);
+        	const FOperation::MaterialSkeletalMeshBreakArgs Args = program.GetOpArgs<FOperation::MaterialSkeletalMeshBreakArgs>(at);
         	f(Args.SkeletalMesh);
         	break;
         }
diff --git a/Engine/Plugins/Mutable/Source/MutableRuntime/Private/MuR/System.cpp b/Engine/Plugins/Mutable/Source/MutableRuntime/Private/MuR/System.cpp
--- a/Engine/Plugins/Mutable/Source/MutableRuntime/Private/MuR/System.cpp
+++ b/Engine/Plugins/Mutable/Source/MutableRuntime/Private/MuR/System.cpp
@@ -154,12 +154,14 @@ namespace UE::Mutable::Private
 
 		InstanceData->UpdateLogger->LogUpdateMessage(TEXT("Creating Mutable instance"), ELogVerbosity::Verbose);
     	
-    	WorkingMemoryManager.LiveInstances.Add(InstanceData.ToWeakPtr());
-    	WorkingMemoryManager.LiveInstances.RemoveAllSwap([](const TWeakPtr<FLiveInstance>& Instance)
 		{
-    		return !Instance.IsValid();
-		});
-    	
+			TScopeLock<FMutex> LockGuard(WorkingMemoryManager.LiveInstancesMutex);
+			WorkingMemoryManager.LiveInstances.Add(InstanceData.ToWeakPtr());
+			WorkingMemoryManager.LiveInstances.RemoveAllSwap([](const TWeakPtr<FLiveInstance>& Instance)
+			{
+				return !Instance.IsValid();
+			});
+		}
 		return InstanceData;
 	}
 
@@ -186,8 +188,6 @@ namespace UE::Mutable::Private
     	InstanceData->SkeletalMeshIdRegistry = MakeShared<FSkeletalMeshIdRegistry>();
     	InstanceData->UpdateLogger = MakeShared<FLiveInstanceLogger>();
 
-    	InstanceData->Cache->Init(InModel->GetProgram().OpAddress.Num());
-    	
     	return InstanceData;
     }
 
@@ -244,28 +244,34 @@ namespace UE::Mutable::Private
 
     	// Program cache setup
     	{
-			// Remove all unlocked data (data from the state used to Prepare the cache)
-			LiveInstance->Cache->ClearLayer0();
-			
+			FProgramCache::EClearFlags ClearFlags = FProgramCache::EClearFlags::Unlocked;
 			if (bClearCacheLayer1)
 			{
-				// Remove all locked data (data from other states)
-				LiveInstance->Cache->ClearLayer1();
+				EnumAddFlags(ClearFlags, FProgramCache::EClearFlags::Locked);
 			}
-			
+
+			LiveInstance->Cache->Clear(ClearFlags);
+
 			// Prepare instance cache
-			LiveInstance->Cache->Prepare(Program, LiveInstance->State);
+			if (Program.States.IsValidIndex(LiveInstanceState))
+			{
+				const FProgram::FState& State = Program.States[LiveInstanceState];
+				for (uint32 Address : State.m_updateCache)
+				{
+					LiveInstance->Cache->LockAddress(FCacheAddress(Address, 0, 0));
+				}
+			}
     	}
 
     	
     	TManagedPtr<const FInstance> Result;
 		{
-			const OP::ADDRESS RootAt = Program.States[LiveInstanceState].Root;
+			const FOperation::ADDRESS RootAt = Program.States[LiveInstanceState].Root;
 			
 			RunCode(LiveInstance, RootAt, InLodMask);
 			if (!bUnrecoverableError)
 			{
-				Result = LiveInstance->Cache->GetInstance(FCacheAddress(RootAt, 0, 0));
+				Result = LiveInstance->Cache->LoadInstance(FCacheAddress(RootAt, 0, 0));
 			}
 
 			// Debug check to see if we managed the op-hit-counts correctly
@@ -297,7 +303,7 @@ namespace UE::Mutable::Private
 			return BlackImage();
     	}
 
-		OP::ADDRESS RootAddress = ImageId.GetKey()->Address;
+		FOperation::ADDRESS RootAddress = ImageId.GetKey()->Address;
 		TManagedPtr<const FImage> Result = BuildImage(LiveInstance, RootAddress, MipsToSkip);
 
 		if (!Result)
@@ -325,7 +331,7 @@ namespace UE::Mutable::Private
 			UE::Tasks::EExtendedTaskPriority::Inline);
     	}
 
-    	OP::ADDRESS RootAddress = ImageId.GetKey()->Address;
+    	FOperation::ADDRESS RootAddress = ImageId.GetKey()->Address;
     	
     	EOpType OpType = LiveInstance->Model->GetProgram().GetOpType(RootAddress);
     	if (GetOpDataType(OpType) != EDataType::Image)
@@ -370,10 +376,12 @@ namespace UE::Mutable::Private
 		{
 			TManagedPtr<const FImage> Result;
 
-			System->bUnrecoverableError = Context->Runner->bUnrecoverableError;
-			if (!Context->Runner->bUnrecoverableError)
+			FCacheAddress RootCacheAddress = FCacheAddress(RootAddress, 0, MipsToSkip);
+			bool bAborted = LiveInstance->Cache->IsAborted(RootCacheAddress);
+			
+			if (!bAborted)
 			{
-				Result = Context->Runner->LoadImage(FCacheAddress(RootAddress, 0, MipsToSkip));
+				Result = LiveInstance->Cache->LoadImage(RootCacheAddress);
 			}
 
 			if (!Result)
@@ -381,6 +389,7 @@ namespace UE::Mutable::Private
 				Result = MakeManaged<FImage>(16, 16, 1, EImageFormat::RGBA_UByte, EInitializationType::Black);
 			}
 			
+			System->bUnrecoverableError = bAborted;
 			return Result;
 		},
 		UE::Tasks::Prerequisites(RunnerTask),
@@ -410,7 +419,7 @@ namespace UE::Mutable::Private
 		
 		FExtendedImageDesc Result;
 
-		OP::ADDRESS RootAddress = ImageId.GetKey()->Address;
+		FOperation::ADDRESS RootAddress = ImageId.GetKey()->Address;
 
 		const UE::Mutable::Private::FModel* Model = LiveInstance->Model.Get();
 		const UE::Mutable::Private::FProgram& Program = Model->GetProgram();
@@ -419,7 +428,7 @@ namespace UE::Mutable::Private
 		int32 VarValue = CVarClearImageDescCache.GetValueOnAnyThread();
 		if (VarValue != 0)
 		{
-			LiveInstance->Cache->ClearDescCache();
+			LiveInstance->Cache->Clear(FProgramCache::EClearFlags::ImageDesc);
 		}
 
 		UE::Mutable::Private::EOpType OpType = Program.GetOpType(RootAddress);
@@ -466,7 +475,7 @@ namespace UE::Mutable::Private
 			UE::Tasks::EExtendedTaskPriority::Inline);
 		}
 		
-		OP::ADDRESS RootAddress = ImageId.GetKey()->Address;
+		FOperation::ADDRESS RootAddress = ImageId.GetKey()->Address;
 		
 		EOpType OpType = LiveInstance->Model->GetProgram().GetOpType(RootAddress);
 		if (GetOpDataType(OpType) != EDataType::Image)
@@ -487,19 +496,19 @@ namespace UE::Mutable::Private
 		};
 		TSharedRef<FContext> Context = MakeShared<FContext>();
 		
+		int8 ExecutionOptions = 0;
 		UE::Tasks::FTask RunnerTask = UE::Tasks::Launch(TEXT("FSystem::GetImageDesc"),
-		[Context, System = this, LiveInstance, RootAddress]()
+		[Context, System = this, LiveInstance, RootAddress, ExecutionOptions]()
 		{
 				
 			// TODO: It should be possible to reuse this data if cleared in the correct places only, together with HeapImageDesc.
 			int32 VarValue = CVarClearImageDescCache.GetValueOnAnyThread();
 			if (VarValue != 0)
 			{
-				LiveInstance->Cache->ClearDescCache();
+				LiveInstance->Cache->Clear(FProgramCache::EClearFlags::ImageDesc);
 			}
 
 			// GetImageDesc may call normal execution paths where meshes are computed.
-			int8 ExecutionOptions = 0;
 			Context->Runner = CodeRunner::Create(
 			LiveInstance,
 				System->Settings,
@@ -518,9 +527,9 @@ namespace UE::Mutable::Private
 		UE::Tasks::ETaskPriority::Inherit);
 		
 		return UE::Tasks::Launch(TEXT("FSystem::GetImageDescResultTask"),
-		[Context, System = this, RootAddress, LiveInstance]() -> FExtendedImageDesc
+		[Context, System = this, RootAddress, LiveInstance, ExecutionOptions]() -> FExtendedImageDesc
 		{
-			FExtendedImageDesc Result = Context->Runner->GetImageDescResult(RootAddress);
+			FExtendedImageDesc Result = Context->Runner->LoadImageDesc(FCacheAddress(RootAddress, 0, ExecutionOptions, FScheduledOp::EType::ImageDesc));
 
 			return Result;
 		},
@@ -541,7 +550,7 @@ namespace UE::Mutable::Private
 			return MakeManaged<FMesh>();
 		}
 
-		OP::ADDRESS RootAddress = MeshId.GetKey()->Address;
+		FOperation::ADDRESS RootAddress = MeshId.GetKey()->Address;
 		TManagedPtr<const FMesh> Result = BuildMesh(
 			LiveInstance,
 			RootAddress,
@@ -574,7 +583,7 @@ namespace UE::Mutable::Private
 			UE::Tasks::EExtendedTaskPriority::Inline);
 		}
 
-		OP::ADDRESS RootAddress = MeshId.GetKey()->Address;
+		FOperation::ADDRESS RootAddress = MeshId.GetKey()->Address;
 
 		EOpType OpType = LiveInstance->Model->GetProgram().GetOpType(RootAddress);
 		if (GetOpDataType(OpType) != EDataType::Mesh)
@@ -620,13 +629,17 @@ namespace UE::Mutable::Private
 		[Context, System = this, RootAddress, MeshContentFilter, LiveInstance]() -> TManagedPtr<const FMesh>
 		{
 			TManagedPtr<const FMesh> Result;
+	
+			FCacheAddress RootCacheAddress = FCacheAddress(RootAddress, 0, static_cast<uint8>(MeshContentFilter));
 
-			System->bUnrecoverableError = Context->Runner->bUnrecoverableError;
-			if (!Context->Runner->bUnrecoverableError)
+			bool bAborted = LiveInstance->Cache->IsAborted(RootCacheAddress);
+			if (!bAborted)
 			{
-				Result = Context->Runner->LoadMesh(FCacheAddress(RootAddress, 0, static_cast<uint8>(MeshContentFilter)));
+				check(LiveInstance->Cache->IsSet(RootCacheAddress));
+				Result = LiveInstance->Cache->LoadMesh(RootCacheAddress);
 			}
 			
+			System->bUnrecoverableError = bAborted;
 			return Result;
 		},
 		UE::Tasks::Prerequisites(RunnerTask),
@@ -652,7 +665,7 @@ namespace UE::Mutable::Private
 			UE::Tasks::EExtendedTaskPriority::Inline);
 		}
 
-		OP::ADDRESS RootAddress = SkeletalMeshId.GetKey()->Address;
+		FOperation::ADDRESS RootAddress = SkeletalMeshId.GetKey()->Address;
 		
 		UE::Mutable::Private::EOpType OpType = LiveInstance->Model->GetProgram().GetOpType(RootAddress);
 		if (GetOpDataType(OpType) != EDataType::SkeletalMesh)
@@ -698,12 +711,16 @@ namespace UE::Mutable::Private
 		{
 			TManagedPtr<const FSkeletalMesh> Result;
 
-			System->bUnrecoverableError = Context->Runner->bUnrecoverableError;
-			if (!Context->Runner->bUnrecoverableError)
+			FCacheAddress RootCacheAddress = FCacheAddress(RootAddress, 0, static_cast<uint8>(SkeletalMeshContentFilter));
+
+			bool bAborted = LiveInstance->Cache->IsAborted(RootCacheAddress);
+			if (!bAborted)
 			{
-				Result = Context->Runner->LoadSkeletalMesh(FCacheAddress(RootAddress, 0, static_cast<uint8>(SkeletalMeshContentFilter)));
+				check(LiveInstance->Cache->IsSet(RootCacheAddress));
+				Result = LiveInstance->Cache->LoadSkeletalMesh(RootCacheAddress);
 			}
 			
+			System->bUnrecoverableError = bAborted;
 			return Result;
 		},
 		UE::Tasks::Prerequisites(RunnerTask),
@@ -729,7 +746,7 @@ namespace UE::Mutable::Private
 			UE::Tasks::EExtendedTaskPriority::Inline);
 		}
 		
-		OP::ADDRESS RootAddress = MaterialId.GetKey()->Address;
+		FOperation::ADDRESS RootAddress = MaterialId.GetKey()->Address;
 	
 		EOpType OpType = LiveInstance->Model->GetProgram().GetOpType(RootAddress);
 		if (GetOpDataType(OpType) != EDataType::Material)
@@ -775,12 +792,15 @@ namespace UE::Mutable::Private
 		{
 			TManagedPtr<const FMaterial> Result;
 
-			System->bUnrecoverableError = Context->Runner->bUnrecoverableError;
-			if (!Context->Runner->bUnrecoverableError)
+			FCacheAddress RootCacheAddress = FCacheAddress(RootAddress, 0, 0);
+
+			bool bAborted = LiveInstance->Cache->IsAborted(RootCacheAddress);
+			if (!bAborted)
 			{
-				Result = Context->Runner->LoadMaterial(FCacheAddress(RootAddress, 0, 0));
+				Result = Context->Runner->LoadMaterial(RootCacheAddress);
 			}
 			
+			System->bUnrecoverableError = bAborted;
 			return Result;
 		},
 		UE::Tasks::Prerequisites(RunnerTask),
@@ -789,7 +809,7 @@ namespace UE::Mutable::Private
 	}
 	
 
-	void FSystem::EndUpdate(const TSharedRef<FLiveInstance>& LiveInstance, bool bClearRoms)
+	void FSystem::EndUpdate(const TSharedRef<FLiveInstance>& LiveInstance, bool bClearRoms, bool bFreeCache)
 	{
 		LLM_SCOPE_BYNAME(TEXT("MutableRuntime"));
 		MUTABLE_CPUPROFILER_SCOPE(EndUpdate);
@@ -802,8 +822,15 @@ namespace UE::Mutable::Private
 
 		// We need to clear the layer 0 cache, because it contains data that is only valid for the current 
 		// parameter values (unless it is data marked as state cache)
-		LiveInstance->Cache->ClearLayer0();
 		
+		FProgramCache::EClearFlags ClearFlags = FProgramCache::EClearFlags::Unlocked;
+		if (bFreeCache)
+		{
+			EnumAddFlags(ClearFlags, FProgramCache::EClearFlags::Locked | FProgramCache::EClearFlags::FreeMemory); 
+		}
+		LiveInstance->Cache->Clear(ClearFlags);
+	
+
 		// Reduce the cache until it fits the limit.
 		// Not providing a cache at this point as there is no need of keeping the cache active anymore
 		WorkingMemoryManager.EnsureBudgetBelow(0, nullptr);
@@ -829,12 +856,12 @@ namespace UE::Mutable::Private
 				UniqueDiscreteCoveredCodeVisitor<>(InSystem, InLiveInstance, FSystem::AllLODs),
 		    	RelevantParameters(OutRelevantParameters)
         {
-			const OP::ADDRESS At = InLiveInstance->Model->GetProgram().States[0].Root;
+			const FOperation::ADDRESS At = InLiveInstance->Model->GetProgram().States[0].Root;
 
             Run(At);
         }
 
-        bool Visit(OP::ADDRESS At, const FProgram& Program) override
+        bool Visit(FOperation::ADDRESS At, const FProgram& Program) override
         {
             switch (Program.GetOpType(At))
             {
@@ -849,8 +876,8 @@ namespace UE::Mutable::Private
 			case EOpType::MI_PARAMETER:
 			case EOpType::IS_PARAMETER:
             {
-				OP::ParameterArgs args = Program.GetOpArgs<OP::ParameterArgs>(At);
-				OP::ADDRESS ParamIndex = args.variable;
+				FOperation::ParameterArgs args = Program.GetOpArgs<FOperation::ParameterArgs>(At);
+				FOperation::ADDRESS ParamIndex = args.variable;
                 RelevantParameters.Add(ParamIndex);
                 break;
             }
@@ -910,7 +937,7 @@ namespace UE::Mutable::Private
 
 	void FSystem::RunCode(
 		const TSharedRef<FLiveInstance>& InLiveInstance,
-		OP::ADDRESS InCodeRoot,
+		FOperation::ADDRESS InCodeRoot,
 		uint32 InLODs,
 		uint8 ExecutionOptions)
 	{
@@ -927,54 +954,64 @@ namespace UE::Mutable::Private
 		
 		constexpr bool bForceInlineExecution = true;
 		Runner->StartRun(bForceInlineExecution);
-		
-		bUnrecoverableError = Runner->bUnrecoverableError;
+	
+		FCacheAddress CacheAddress = FCacheAddress(InCodeRoot, 0, ExecutionOptions);
+		bool bAborted = InLiveInstance->Cache->IsAborted(CacheAddress);
+
+		// In case of aborted the value will not be read, decrease root address hit count
+		// to balance the ophit added at runner creation.
+		if (bAborted)
+		{
+			InLiveInstance->Cache->UpdateHitCount(CacheAddress, -1);
+		}
+
+		bUnrecoverableError = bAborted;
 	}
 	
 	
-	bool FSystem::BuildBool(const TSharedRef<FLiveInstance>& InLiveInstance, OP::ADDRESS InAddress)
+	bool FSystem::BuildBool(const TSharedRef<FLiveInstance>& InLiveInstance, FOperation::ADDRESS InAddress)
     {
     	RunCode(InLiveInstance, InAddress);
 	
     	bool bResult = false;
     	if (!bUnrecoverableError)
     	{
-    		bResult = InLiveInstance->Cache->GetBool(FCacheAddress(InAddress, 0, 0));
+    		bResult = InLiveInstance->Cache->LoadBool(FCacheAddress(InAddress, 0, 0));
     	}
 	
     	return bResult;    
     }
 	
 	
-	int32 FSystem::BuildInt(const TSharedRef<FLiveInstance>& InLiveInstance, OP::ADDRESS InAddress)
+	int32 FSystem::BuildInt(const TSharedRef<FLiveInstance>& InLiveInstance, FOperation::ADDRESS InAddress)
 	{
 		RunCode(InLiveInstance, InAddress);
 
 		int32 Result = 0;
 		if (!bUnrecoverableError)
 		{
-			Result = InLiveInstance->Cache->GetInt(FCacheAddress(InAddress, 0, 0));
+			Result = InLiveInstance->Cache->LoadInt(FCacheAddress(InAddress, 0, 0));
 		}
 
 		return Result;
 	}
 
 	
-	float FSystem::BuildScalar(const TSharedRef<FLiveInstance>& InLiveInstance, OP::ADDRESS InAddress)
+	float FSystem::BuildScalar(const TSharedRef<FLiveInstance>& InLiveInstance, FOperation::ADDRESS InAddress)
 	{
     	RunCode(InLiveInstance, InAddress);
 
     	float Result = 0.0f;		
     	if (!bUnrecoverableError)
     	{
-    		Result = InLiveInstance->Cache->GetScalar(FCacheAddress(InAddress, 0, 0));
+    		Result = InLiveInstance->Cache->LoadScalar(FCacheAddress(InAddress, 0, 0));
     	}
 
     	return Result;
 	}
 
 
-	FVector4f FSystem::BuildColor(const TSharedRef<FLiveInstance>& InLiveInstance, OP::ADDRESS InAddress)
+	FVector4f FSystem::BuildColor(const TSharedRef<FLiveInstance>& InLiveInstance, FOperation::ADDRESS InAddress)
 	{
 		FVector4f Result(0,0,0,1);
     	
@@ -984,7 +1021,7 @@ namespace UE::Mutable::Private
 			RunCode(InLiveInstance, InAddress);
 			if (!bUnrecoverableError)
 			{
-				Result = InLiveInstance->Cache->GetColor(FCacheAddress(InAddress, 0, 0));
+				Result = InLiveInstance->Cache->LoadColor(FCacheAddress(InAddress, 0, 0));
 			}
 		}
 
@@ -992,21 +1029,21 @@ namespace UE::Mutable::Private
 	}
 
 	
-	FProjector FSystem::BuildProjector(const TSharedRef<FLiveInstance>& InLiveInstance, OP::ADDRESS InAddress)
+	FProjector FSystem::BuildProjector(const TSharedRef<FLiveInstance>& InLiveInstance, FOperation::ADDRESS InAddress)
 	{
     	RunCode(InLiveInstance, InAddress);
 
 		FProjector Result;
 		if (!bUnrecoverableError)
 		{
-			Result = InLiveInstance->Cache->GetProjector(FCacheAddress(InAddress, 0, 0));
+			Result = InLiveInstance->Cache->LoadProjector(FCacheAddress(InAddress, 0, 0));
 		}
 
 		return Result;
 	}
 
 
-	TManagedPtr<const FImage> FSystem::BuildImage(const TSharedRef<FLiveInstance>& InLiveInstance, OP::ADDRESS InAddress, int32 MipsToSkip)
+	TManagedPtr<const FImage> FSystem::BuildImage(const TSharedRef<FLiveInstance>& InLiveInstance, FOperation::ADDRESS InAddress, int32 MipsToSkip)
 	{
 		TManagedPtr<const FImage> Result;
     	
@@ -1016,7 +1053,7 @@ namespace UE::Mutable::Private
 			RunCode(InLiveInstance, InAddress, FSystem::AllLODs, uint8(MipsToSkip));
 			if (!bUnrecoverableError)
 			{
-				Result = InLiveInstance->Cache->GetImage(FCacheAddress(InAddress, 0, MipsToSkip));
+				Result = InLiveInstance->Cache->LoadImage(FCacheAddress(InAddress, 0, MipsToSkip));
 			}			
 		}
 
@@ -1024,7 +1061,7 @@ namespace UE::Mutable::Private
 	}
 
 
-	TManagedPtr<const FMesh> FSystem::BuildMesh(const TSharedRef<FLiveInstance>& InLiveInstance, OP::ADDRESS InRootAddress, EMeshContentFlags MeshContentFilter)
+	TManagedPtr<const FMesh> FSystem::BuildMesh(const TSharedRef<FLiveInstance>& InLiveInstance, FOperation::ADDRESS InRootAddress, EMeshContentFlags MeshContentFilter)
 	{
 		TManagedPtr<const FMesh> Result;
     	
@@ -1035,7 +1072,7 @@ namespace UE::Mutable::Private
 			RunCode(InLiveInstance, InRootAddress, FSystem::AllLODs, ExecutionOptions);
 			if (!bUnrecoverableError)
 			{
-				Result = InLiveInstance->Cache->GetMesh(FCacheAddress(InRootAddress, 0, ExecutionOptions));
+				Result = InLiveInstance->Cache->LoadMesh(FCacheAddress(InRootAddress, 0, ExecutionOptions));
 			}
 		}
 
@@ -1043,7 +1080,7 @@ namespace UE::Mutable::Private
 	}
 
 
-	TManagedPtr<const FInstance> FSystem::BuildInstance(const TSharedRef<FLiveInstance>& InLiveInstance, OP::ADDRESS InAddress)
+	TManagedPtr<const FInstance> FSystem::BuildInstance(const TSharedRef<FLiveInstance>& InLiveInstance, FOperation::ADDRESS InAddress)
 	{
 		TManagedPtr<const FInstance> Result;
     	
@@ -1053,7 +1090,7 @@ namespace UE::Mutable::Private
 			RunCode(InLiveInstance, InAddress);
 			if (!bUnrecoverableError)
 			{
-				Result = InLiveInstance->Cache->GetInstance(FCacheAddress(InAddress, 0, 0));
+				Result = InLiveInstance->Cache->LoadInstance(FCacheAddress(InAddress, 0, 0));
 			}
 		}
 
@@ -1061,7 +1098,7 @@ namespace UE::Mutable::Private
 	}
 
 
-	TManagedPtr<const FLayout> FSystem::BuildLayout(const TSharedRef<FLiveInstance>& InLiveInstance, OP::ADDRESS InAddress)
+	TManagedPtr<const FLayout> FSystem::BuildLayout(const TSharedRef<FLiveInstance>& InLiveInstance, FOperation::ADDRESS InAddress)
 	{
 		TManagedPtr<const FLayout> Result;
     	
@@ -1073,7 +1110,7 @@ namespace UE::Mutable::Private
 				RunCode(InLiveInstance, InAddress);
 				if (!bUnrecoverableError)
 				{
-					Result = InLiveInstance->Cache->GetLayout(FCacheAddress(InAddress, 0, 0));
+					Result = InLiveInstance->Cache->LoadLayout(FCacheAddress(InAddress, 0, 0));
 				}
 			}
 		}
@@ -1082,7 +1119,7 @@ namespace UE::Mutable::Private
 	}
 
 
-	TManagedPtr<const String> FSystem::BuildString(const TSharedRef<FLiveInstance>& InLiveInstance, OP::ADDRESS InAddress)
+	TManagedPtr<const String> FSystem::BuildString(const TSharedRef<FLiveInstance>& InLiveInstance, FOperation::ADDRESS InAddress)
 	{
 		TManagedPtr<const String> Result;
     	
@@ -1094,7 +1131,7 @@ namespace UE::Mutable::Private
 				RunCode(InLiveInstance, InAddress);
 				if (!bUnrecoverableError)
 				{
-					Result = InLiveInstance->Cache->GetString(FCacheAddress(InAddress, 0, 0));
+					Result = InLiveInstance->Cache->LoadString(FCacheAddress(InAddress, 0, 0));
 				}
 			}
 		}
@@ -1120,7 +1157,7 @@ namespace UE::Mutable::Private
 	int32 FWorkingMemoryManager::GetRomBytes() const
     {
     	int32 Result = 0;
-
+		TScopeLock<FMutex> LockGuard(LiveInstancesMutex);
     	for (const TWeakPtr<FLiveInstance>& WeakInstance : LiveInstances)
     	{
     		TSharedPtr<FLiveInstance> Instance = WeakInstance.Pin();
@@ -1181,35 +1218,7 @@ namespace UE::Mutable::Private
 		TRACE_COUNTER_SET(MutableRuntime_MemTotal,    TotalBytes);
 		TRACE_COUNTER_SET(MutableRuntime_MemBudget,   BudgetBytes);
 #endif
-	}
-
-
-	void FProgramCache::Prepare(const FProgram& InProgram, int32 InState)
-	{
-    	MUTABLE_CPUPROFILER_SCOPE(PrepareCache);
-			
-    	const int32 OpCount = InProgram.OpAddress.Num();
-    	Init(OpCount);
-
-    	// Clear cache flags of existing data
-    	CodeContainer<FProgramCache::FOpExecutionData>::iterator It = OpExecutionData.begin();
-    	for (; It.IsValid(); ++It)
-    	{
-    		(*It).OpHitCount = 0; // This should already be 0, but just in case.
-    		(*It).IsCacheLocked = false;
-    	}
-
-    	// Mark the resources that have to be cached to update the instance in this state
-    	if (InState >= 0 && InState< InProgram.States.Num())
-    	{
-    		const FProgram::FState& State = InProgram.States[InState];
-    		for (uint32 Address : State.m_updateCache)
-    		{
-    			SetForceCached(Address);
-    		}
-    	}
-	}
-	
+	}	
 
 	int64 FWorkingMemoryManager::GetCurrentMemoryBytes() const
 	{
@@ -1257,6 +1266,7 @@ namespace UE::Mutable::Private
 		// Try to free loaded roms
 		if (!bFinished)
 		{
+			TScopeLock<FMutex> LockGuard(LiveInstancesMutex);
 			for (TWeakPtr<FLiveInstance>& WeakLiveInstance : LiveInstances)
 			{
 				TSharedPtr<FLiveInstance> LiveInstance = WeakLiveInstance.Pin();
@@ -1292,67 +1302,15 @@ namespace UE::Mutable::Private
 
 		auto TryRemoveFromCache = [&bFinished, &TotalBytes, Budget = BudgetBytes](FProgramCache& Cache)
 		{
-			int64 ImageRemovedBytes = 0;
-
-			for (const TManagedWeakPtr<const FImage>& ImageResult : Cache.ImageResults)
-			{
-				TManagedPtr<const FImage> StrongImagePtr = ImageResult.Pin();
-				if (!StrongImagePtr)
-				{
-					continue;
-				}
-
-				int32 ImageDataSize = StrongImagePtr->GetDataSize();
-				if (StrongImagePtr.TryDeleteObject())
-				{
-					ImageRemovedBytes += ImageDataSize;
-				}
-
-				bFinished = TotalBytes - ImageRemovedBytes <= Budget;
-
-				if (bFinished)
-				{
-					break;
-				}
-			}
-
-			TotalBytes -= ImageRemovedBytes;
-			if (bFinished)
-			{
-				return;
-			}
-
-			int64 MeshRemovedBytes = 0;
-			for (const TManagedWeakPtr<const FMesh>& MeshResult : Cache.MeshResults)
-			{
-				TManagedPtr<const FMesh> StrongMeshPtr = MeshResult.Pin();
-				if (!StrongMeshPtr)
-				{
-					continue;
-				}
-
-				int32 MeshDataSize = StrongMeshPtr->GetDataSize();
-				if (StrongMeshPtr.TryDeleteObject())
-				{
-					MeshRemovedBytes += MeshDataSize;
-				}
-
-				bFinished = TotalBytes - MeshRemovedBytes <= Budget;
-
-				if (bFinished)
-				{
-					break;
-				}
-			}
-
-			TotalBytes -= MeshRemovedBytes;
+			TotalBytes -= Cache.TryFreeWeakMemory(TotalBytes - Budget);
+			bFinished = TotalBytes <= Budget;
 		};
 
 		// Try to free cache 1 memory
 		if (!bFinished)
 		{
 			MUTABLE_CPUPROFILER_SCOPE(EnsureBudgetBelow_FreeCached);	
-			// From other live instances first
+			TScopeLock<FMutex> LockGuard(LiveInstancesMutex);
 			for (const TWeakPtr<FLiveInstance>& WeakInstance : LiveInstances)
 			{
 				TSharedPtr<FLiveInstance> Instance = WeakInstance.Pin();
@@ -1428,7 +1386,7 @@ namespace UE::Mutable::Private
 	}
 
 
-	void GetResourceKey(const FModel& Model, const FParameters& Params, uint32 ParamListIndex, OP::ADDRESS RootAt, TMemoryTrackedArray<uint8>& ParameterValuesBlob)
+	void GetResourceKey(const FModel& Model, const FParameters& Params, uint32 ParamListIndex, FOperation::ADDRESS RootAt, TMemoryTrackedArray<uint8>& ParameterValuesBlob)
 	{
 		MUTABLE_CPUPROFILER_SCOPE(GetResourceKey);
     	
@@ -1727,7 +1685,7 @@ namespace UE::Mutable::Private
 	}
 
 
-	FMeshId FSystem::GetMeshId(const TSharedPtr<FMeshIdRegistry>& Registry, const FModel& Model, const FParameters& Params, OP::ADDRESS RootAt)
+	FMeshId FSystem::GetMeshId(const TSharedPtr<FMeshIdRegistry>& Registry, const FModel& Model, const FParameters& Params, FOperation::ADDRESS RootAt)
     {
     	check(Registry)
     	
@@ -1742,7 +1700,7 @@ namespace UE::Mutable::Private
     }
 	
     	
-	FImageId FSystem::GetImageId(const TSharedPtr<FImageIdRegistry>& Registry, const FModel& Model, const FParameters& Params, OP::ADDRESS RootAt)
+	FImageId FSystem::GetImageId(const TSharedPtr<FImageIdRegistry>& Registry, const FModel& Model, const FParameters& Params, FOperation::ADDRESS RootAt)
     {
     	check(Registry)
     
@@ -1757,7 +1715,7 @@ namespace UE::Mutable::Private
     }
 
 	
-	FMaterialId FSystem::GetMaterialId(const TSharedPtr<FMaterialIdRegistry>& Registry, const FModel& Model, const FParameters& Params, OP::ADDRESS RootAt)
+	FMaterialId FSystem::GetMaterialId(const TSharedPtr<FMaterialIdRegistry>& Registry, const FModel& Model, const FParameters& Params, FOperation::ADDRESS RootAt)
     {
     	check(Registry)
     
@@ -1772,7 +1730,7 @@ namespace UE::Mutable::Private
     }
 
 	
-	FSkeletalMeshId FSystem::GetSkeletalMeshId(const TSharedPtr<FSkeletalMeshIdRegistry>& Registry, const FModel& Model, const FParameters& Params, OP::ADDRESS RootAt)
+	FSkeletalMeshId FSystem::GetSkeletalMeshId(const TSharedPtr<FSkeletalMeshIdRegistry>& Registry, const FModel& Model, const FParameters& Params, FOperation::ADDRESS RootAt)
 	{
 		check(Registry)
 	
diff --git a/Engine/Plugins/Mutable/Source/MutableTools/Internal/MuT/AST.h b/Engine/Plugins/Mutable/Source/MutableTools/Internal/MuT/AST.h
--- a/Engine/Plugins/Mutable/Source/MutableTools/Internal/MuT/AST.h
+++ b/Engine/Plugins/Mutable/Source/MutableTools/Internal/MuT/AST.h
@@ -297,7 +297,7 @@ namespace UE::Mutable::Private
 
 		FAdditionalData AdditionalData;
 
-		TMap<OP::ADDRESS, FInstancedStruct>* ExternalOperations = nullptr;
+		TMap<FOperation::ADDRESS, FInstancedStruct>* ExternalOperations = nullptr;
 	};
 
 
@@ -629,7 +629,7 @@ namespace UE::Mutable::Private
 		/** Convert the operation graph at Root into code in the given program.
 		* Potentially destroys the data in this operation, so it shouldn't be used after calling Link.
 		*/
-		static OP::ADDRESS FullLink( Ptr<ASTOp>& Root, FProgram&, FLinkerOptions* );
+		static FOperation::ADDRESS FullLink( Ptr<ASTOp>& Root, FProgram&, FLinkerOptions* );
 
     private:
 
@@ -692,7 +692,7 @@ namespace UE::Mutable::Private
         //!
 		static void LinkRange(FProgram& program,
 			const FRangeData& range,
-			OP::ADDRESS& rangeSize,
+			FOperation::ADDRESS& rangeSize,
 			uint16& rangeId);
 
     public:
@@ -808,7 +808,7 @@ namespace UE::Mutable::Private
         );
 
         //!
-        OP::ADDRESS linkedAddress = 0;
+        FOperation::ADDRESS LinkedAddress = 0;
 
         //! Generic traverse control counter. It should always be left to 0 after any process for all
         //! nodes in the hierarchy.
diff --git a/Engine/Plugins/Mutable/Source/MutableTools/Internal/MuT/Compiler.h b/Engine/Plugins/Mutable/Source/MutableTools/Internal/MuT/Compiler.h
--- a/Engine/Plugins/Mutable/Source/MutableTools/Internal/MuT/Compiler.h
+++ b/Engine/Plugins/Mutable/Source/MutableTools/Internal/MuT/Compiler.h
@@ -150,7 +150,7 @@ namespace UE::Mutable::Private
         UE_API Compiler( Ptr<CompilerOptions> Options, TFunction<void()>& InWaitCallback );
 
         //! Compile the expression into a run-time model.
-		UE_API TSharedPtr<FModel> Compile( const Ptr<Node>& pNode, TMap<OP::ADDRESS, FInstancedStruct>& ExternalOperations);
+		UE_API TSharedPtr<FModel> Compile( const Ptr<Node>& pNode, TMap<FOperation::ADDRESS, FInstancedStruct>& ExternalOperations);
 
         //! Return the log of messages of all the compile operations executed so far.
         UE_API TSharedPtr<FErrorLog> GetLog() const;
diff --git a/Engine/Plugins/Mutable/Source/MutableTools/Private/MuT/AST.cpp b/Engine/Plugins/Mutable/Source/MutableTools/Private/MuT/AST.cpp
--- a/Engine/Plugins/Mutable/Source/MutableTools/Private/MuT/AST.cpp
+++ b/Engine/Plugins/Mutable/Source/MutableTools/Private/MuT/AST.cpp
@@ -243,15 +243,24 @@ UE::Mutable::Private::Ptr<ASTOp> ASTOp::DeepClone( const Ptr<ASTOp>& root )
 
 
 //-------------------------------------------------------------------------------------------------
-OP::ADDRESS ASTOp::FullLink( Ptr<ASTOp>& Root, FProgram& Program, FLinkerOptions* Options )
+FOperation::ADDRESS ASTOp::FullLink( Ptr<ASTOp>& Root, FProgram& Program, FLinkerOptions* Options )
 {
     MUTABLE_CPUPROFILER_SCOPE(AST_FullLink);
 
-    Traverse_BottomUp_Unique( Root,
-                              [&](Ptr<ASTOp> n){ n->Link(Program, Options); },
-                              [&](Ptr<const ASTOp> n){ return n->linkedAddress==0; });
+    auto Link = [&Program, Options](const Ptr<ASTOp>& Node)
+    {
+        Node->Link(Program, Options);
+		check(!(GetAddressDataType(Node->LinkedAddress) == EDataType::None) || Node->LinkedAddress == 0);
+    };
+
+    auto Accept = [](const Ptr<const ASTOp>& Node)
+    {
+        return Node->LinkedAddress == 0;
+    };
+
+    Traverse_BottomUp_Unique(Root, Link, Accept); 
 
-	OP::ADDRESS Result = Root->linkedAddress;
+	FOperation::ADDRESS Result = Root->LinkedAddress;
 
 	// This signals the caller that the Root pointer shouldn't be used anymore.
 	Root = nullptr;
@@ -927,39 +936,37 @@ ASTOp::EClosedMeshTest ASTOp::IsClosedMesh(TMap<const ASTOp*, EClosedMeshTest>*
 }
 
 
-void ASTOp::LinkRange(FProgram& program,
-	const FRangeData& range,
-	OP::ADDRESS& rangeSize,
-	uint16& rangeId)
+void ASTOp::LinkRange(FProgram& Program, const FRangeData& Range, FOperation::ADDRESS& RangeSize, uint16& RangeId)
 {
-	if (range.rangeSize)
+	if (Range.rangeSize)
 	{
-		if (range.rangeSize->linkedRange < 0)
+		if (Range.rangeSize->linkedRange < 0)
 		{
-			check(program.Ranges.Num() < 255);
-			range.rangeSize->linkedRange = int8(program.Ranges.Num());
-			FRangeDesc rangeData;
-			rangeData.Name = range.rangeName;
-			rangeData.UID = range.rangeUID;
+			check(Program.Ranges.Num() < 255);
+			Range.rangeSize->linkedRange = int8(Program.Ranges.Num());
+			FRangeDesc RangeData;
+			RangeData.Name = Range.rangeName;
+			RangeData.UID = Range.rangeUID;
 
 			// Try to see if a parameter directly controls de size of the range. This is used
 			// to store hint data for instance generation in tools or randomizers that want to
 			// support multilayer, but it is not critical otherwise.
 			int32 EstimatedSizeParameter = -1;
-			if ( range.rangeSize->GetOpType() == EOpType::SC_PARAMETER
-				||
-				range.rangeSize->GetOpType() == EOpType::NU_PARAMETER )
+			if (Range.rangeSize->GetOpType() == EOpType::SC_PARAMETER ||
+				Range.rangeSize->GetOpType() == EOpType::NU_PARAMETER )
 			{
-				const ASTOpParameter* ParamOp = static_cast<const ASTOpParameter*>(range.rangeSize.child().get());
+				const ASTOpParameter* ParamOp = static_cast<const ASTOpParameter*>(Range.rangeSize.child().get());
 
 				EstimatedSizeParameter = ParamOp->LinkedParameterIndex;
 			}
-			rangeData.DimensionParameter = EstimatedSizeParameter;
 
-			program.Ranges.Add(rangeData);
+			RangeData.DimensionParameter = EstimatedSizeParameter;
+
+			Program.Ranges.Add(RangeData);
 		}
-		rangeSize = range.rangeSize->linkedAddress;
-		rangeId = uint16(range.rangeSize->linkedRange);
+
+		RangeSize = Range.rangeSize->LinkedAddress;
+		RangeId = uint16(Range.rangeSize->linkedRange);
 	}
 }
 
diff --git a/Engine/Plugins/Mutable/Source/MutableTools/Private/MuT/ASTOpAddExtensionData.cpp b/Engine/Plugins/Mutable/Source/MutableTools/Private/MuT/ASTOpAddExtensionData.cpp
--- a/Engine/Plugins/Mutable/Source/MutableTools/Private/MuT/ASTOpAddExtensionData.cpp
+++ b/Engine/Plugins/Mutable/Source/MutableTools/Private/MuT/ASTOpAddExtensionData.cpp
@@ -68,28 +68,29 @@ void ASTOpAddExtensionData::ForEachChild(const TFunctionRef<void(ASTChild&)> F)
 void ASTOpAddExtensionData::Link(FProgram& Program, FLinkerOptions*)
 {
 	// Already linked?
-	if (linkedAddress)
+	if (LinkedAddress)
 	{
 		return;
 	}
 
-	OP::InstanceAddExtensionDataArgs Args;
+	FOperation::InstanceAddExtensionDataArgs Args;
 	FMemory::Memzero(Args);
 
-	if (!Instance || !Instance->linkedAddress)
+	if (!Instance || !Instance->LinkedAddress)
 	{
 		// Can happen if there's no reference skeletal mesh in the first component
 		return;
 	}
 
-	Args.Instance = Instance->linkedAddress;
+	Args.Instance = Instance->LinkedAddress;
 
-	check(ExtensionData->linkedAddress);
-	Args.ExtensionData = ExtensionData->linkedAddress;
+	check(ExtensionData->LinkedAddress);
+	Args.ExtensionData = ExtensionData->LinkedAddress;
 
-	linkedAddress = (OP::ADDRESS)Program.OpAddress.Num();
-	Program.OpAddress.Add((uint32_t)Program.ByteCode.Num());
-	AppendCode(Program.ByteCode, EOpType::IN_ADDEXTENSIONDATA);
+	++Program.NumOps;
+	LinkedAddress = MakeProgramAddress(GetOpType(), Program.ByteCode.Num());
+
+	AppendCode(Program.ByteCode, GetOpType());
 	AppendCode(Program.ByteCode, Args);
 }
 
diff --git a/Engine/Plugins/Mutable/Source/MutableTools/Private/MuT/ASTOpAddLOD.cpp b/Engine/Plugins/Mutable/Source/MutableTools/Private/MuT/ASTOpAddLOD.cpp
--- a/Engine/Plugins/Mutable/Source/MutableTools/Private/MuT/ASTOpAddLOD.cpp
+++ b/Engine/Plugins/Mutable/Source/MutableTools/Private/MuT/ASTOpAddLOD.cpp
@@ -70,21 +70,20 @@ namespace UE::Mutable::Private
 	}
 
 
-	void ASTOpAddLOD::Link(FProgram& program, FLinkerOptions*)
+	void ASTOpAddLOD::Link(FProgram& Program, FLinkerOptions*)
 	{
 		// Already linked?
-		if (!linkedAddress)
+		if (!LinkedAddress)
 		{
-			linkedAddress = (OP::ADDRESS)program.OpAddress.Num();
-			program.OpAddress.Add((uint32)program.ByteCode.Num());
-			AppendCode(program.ByteCode, EOpType::IN_ADDLOD);
+			check(Program.ByteCode.Num() < FOperation::MaxAddressByteCodeOffset);
+			LinkedAddress = MakeProgramAddress(GetOpType(), (FOperation::ADDRESS)Program.ByteCode.Num());
 
 			// Calculate LODCount mimicking previous behaviour.
 			// It may be better to leave null LODs if they are null in the AST.
 			uint8 LODCount = 0;
 			for (const ASTChild& LOD : lods)
 			{
-				OP::ADDRESS LODAddress = 0;
+				FOperation::ADDRESS LODAddress = 0;
 				if (LOD)
 				{
 					++LODCount;
@@ -96,14 +95,19 @@ namespace UE::Mutable::Private
 				}
 			}
 
-			AppendCode(program.ByteCode, LODCount);
+			++Program.NumOps;
+			AppendCode(Program.ByteCode, GetOpType()); 
+			AppendCode(Program.ByteCode, LODCount);
+
 			for (const ASTChild& LOD : lods)
 			{
-				OP::ADDRESS LODAddress = 0;
+				FOperation::ADDRESS LODAddress = 0;
 				if (LOD)
 				{
-					LODAddress = LOD->linkedAddress;
-					AppendCode(program.ByteCode, LODAddress);
+					LODAddress = LOD->LinkedAddress;
+
+					AppendCode(Program.ByteCode, LODAddress);
+					
 					if (LODCount == 255)
 					{
 						break;
@@ -111,7 +115,6 @@ namespace UE::Mutable::Private
 				}
 			}
 		}
-
 	}
 
 }
diff --git a/Engine/Plugins/Mutable/Source/MutableTools/Private/MuT/ASTOpAddMaterial.cpp b/Engine/Plugins/Mutable/Source/MutableTools/Private/MuT/ASTOpAddMaterial.cpp
--- a/Engine/Plugins/Mutable/Source/MutableTools/Private/MuT/ASTOpAddMaterial.cpp
+++ b/Engine/Plugins/Mutable/Source/MutableTools/Private/MuT/ASTOpAddMaterial.cpp
@@ -11,17 +11,14 @@ namespace UE::Mutable::Private
 {
 
 
-//-------------------------------------------------------------------------------------------------
-//-------------------------------------------------------------------------------------------------
-//-------------------------------------------------------------------------------------------------
-	ASTOpAddMaterial::ASTOpAddMaterial()
+ASTOpAddMaterial::ASTOpAddMaterial()
 	: Instance(this)
 	, Material(this)
 {
 }
 
 
-	ASTOpAddMaterial::~ASTOpAddMaterial()
+ASTOpAddMaterial::~ASTOpAddMaterial()
 {
 	// Explicit call needed to avoid recursive destruction
 	ASTOp::RemoveChildren();
@@ -71,22 +68,22 @@ void ASTOpAddMaterial::Link(FProgram& Program, FLinkerOptions*)
 	check(Type != EOpType::NONE);
 
 	// Already linked?
-	if (linkedAddress)
+	if (LinkedAddress)
 	{
 		return;
 	}
 
-	OP::InstanceAddMaterialArgs Args;
+	FOperation::InstanceAddMaterialArgs Args;
 	FMemory::Memzero(Args);
 
 	if (Instance)
 	{
-		Args.Instance = Instance->linkedAddress;
+		Args.Instance = Instance->LinkedAddress;
 	}
 
 	if (Material)
 	{
-		Args.Material = Material->linkedAddress;
+		Args.Material = Material->LinkedAddress;
 		
 		// Find out relevant parameters. \todo: this may be optimised by reusing partial
 		// values in a LINK_CONTEXT or similar
@@ -113,18 +110,19 @@ void ASTOpAddMaterial::Link(FProgram& Program, FLinkerOptions*)
 
 		if (It != INDEX_NONE)
 		{
-			Program.RelevantParameterList.Add(Material->linkedAddress, It);
+			Program.RelevantParameterList.Add(Material->LinkedAddress, It);
 		}
 		else
 		{
-			Program.RelevantParameterList.Add(Material->linkedAddress, Program.ParameterLists.Num());
+			Program.RelevantParameterList.Add(Material->LinkedAddress, Program.ParameterLists.Num());
 			Program.ParameterLists.Add(Params);
 		}
 	}
 
-	linkedAddress = (OP::ADDRESS)Program.OpAddress.Num();
-	Program.OpAddress.Add((uint32_t)Program.ByteCode.Num());
-	AppendCode(Program.ByteCode, Type);
+	++Program.NumOps;
+	LinkedAddress = MakeProgramAddress(GetOpType(), Program.ByteCode.Num());
+
+	AppendCode(Program.ByteCode, GetOpType());
 	AppendCode(Program.ByteCode, Args);
 }
 
diff --git a/Engine/Plugins/Mutable/Source/MutableTools/Private/MuT/ASTOpAddOverlayMaterial.cpp b/Engine/Plugins/Mutable/Source/MutableTools/Private/MuT/ASTOpAddOverlayMaterial.cpp
--- a/Engine/Plugins/Mutable/Source/MutableTools/Private/MuT/ASTOpAddOverlayMaterial.cpp
+++ b/Engine/Plugins/Mutable/Source/MutableTools/Private/MuT/ASTOpAddOverlayMaterial.cpp
@@ -76,22 +76,22 @@ namespace UE::Mutable::Private
 	void ASTOpAddOverlayMaterial::Link(FProgram& Program, FLinkerOptions*)
 	{
 		// Already linked?
-		if (linkedAddress)
+		if (LinkedAddress)
 		{
 			return;
 		}
 
-		OP::InstanceAddOverlayMaterialArgs Args;
+		FOperation::InstanceAddOverlayMaterialArgs Args;
 		FMemory::Memzero(Args);
 
 		if (Instance)
 		{
-			Args.Instance = Instance->linkedAddress;
+			Args.Instance = Instance->LinkedAddress;
 		}
 
 		if (Material)
 		{
-			Args.Material = Material->linkedAddress;
+			Args.Material = Material->LinkedAddress;
 			
 			// Find out relevant parameters. \todo: this may be optimised by reusing partial
 			// values in a LINK_CONTEXT or similar
@@ -118,15 +118,16 @@ namespace UE::Mutable::Private
 
 			if (It != INDEX_NONE)
 			{
-				Program.RelevantParameterList.Add(Material->linkedAddress, It);
+				Program.RelevantParameterList.Add(Material->LinkedAddress, It);
 			}
 			else
 			{
-				Program.RelevantParameterList.Add(Material->linkedAddress, Program.ParameterLists.Num());
+				Program.RelevantParameterList.Add(Material->LinkedAddress, Program.ParameterLists.Num());
 				Program.ParameterLists.Add(Params);
 			}
 		}
 		
+
 		if (SlotName.IsNone())
 		{
 			// Mark the SlotName as not relevant so the operation result does not get limited to only one slot
@@ -137,10 +138,11 @@ namespace UE::Mutable::Private
 			// Having no slotName here is valid but that would make the whole component get an overlay material and not the material slot
 			Args.SlotName = Program.AddConstant(SlotName);
 		}
+
+		++Program.NumOps;
+		LinkedAddress = MakeProgramAddress(GetOpType(), Program.ByteCode.Num());
 		
-		linkedAddress = (OP::ADDRESS)Program.OpAddress.Num();
-		Program.OpAddress.Add((uint32_t)Program.ByteCode.Num());
 		AppendCode(Program.ByteCode, GetOpType());
 		AppendCode(Program.ByteCode, Args);
 	}
-}
\ No newline at end of file
+}
diff --git a/Engine/Plugins/Mutable/Source/MutableTools/Private/MuT/ASTOpAddOverrideMaterial.cpp b/Engine/Plugins/Mutable/Source/MutableTools/Private/MuT/ASTOpAddOverrideMaterial.cpp
--- a/Engine/Plugins/Mutable/Source/MutableTools/Private/MuT/ASTOpAddOverrideMaterial.cpp
+++ b/Engine/Plugins/Mutable/Source/MutableTools/Private/MuT/ASTOpAddOverrideMaterial.cpp
@@ -76,22 +76,22 @@ namespace UE::Mutable::Private
 	void ASTOpAddOverrideMaterial::Link(FProgram& Program, FLinkerOptions*)
 	{
 		// Already linked?
-		if (linkedAddress)
+		if (LinkedAddress)
 		{
 			return;
 		}
 
-		OP::InstanceAddOverrideMaterialArgs Args;
+		FOperation::InstanceAddOverrideMaterialArgs Args;
 		FMemory::Memzero(Args);
 
 		if (Instance)
 		{
-			Args.Instance = Instance->linkedAddress;
+			Args.Instance = Instance->LinkedAddress;
 		}
 
 		if (Material)
 		{
-			Args.Material = Material->linkedAddress;
+			Args.Material = Material->LinkedAddress;
 			
 			// Find out relevant parameters. \todo: this may be optimised by reusing partial
 			// values in a LINK_CONTEXT or similar
@@ -118,19 +118,20 @@ namespace UE::Mutable::Private
 
 			if (It != INDEX_NONE)
 			{
-				Program.RelevantParameterList.Add(Material->linkedAddress, It);
+				Program.RelevantParameterList.Add(Material->LinkedAddress, It);
 			}
 			else
 			{
-				Program.RelevantParameterList.Add(Material->linkedAddress, Program.ParameterLists.Num());
+				Program.RelevantParameterList.Add(Material->LinkedAddress, Program.ParameterLists.Num());
 				Program.ParameterLists.Add(Params);
 			}
 		}
 		
 		Args.SlotName = Program.AddConstant(SlotName.ToString());
 
-		linkedAddress = (OP::ADDRESS)Program.OpAddress.Num();
-		Program.OpAddress.Add((uint32_t)Program.ByteCode.Num());
+		++Program.NumOps;
+		LinkedAddress = MakeProgramAddress(GetOpType(), Program.ByteCode.Num());
+
 		AppendCode(Program.ByteCode, GetOpType());
 		AppendCode(Program.ByteCode, Args);
 	}
diff --git a/Engine/Plugins/Mutable/Source/MutableTools/Private/MuT/ASTOpAddSkeletalMesh.cpp b/Engine/Plugins/Mutable/Source/MutableTools/Private/MuT/ASTOpAddSkeletalMesh.cpp
--- a/Engine/Plugins/Mutable/Source/MutableTools/Private/MuT/ASTOpAddSkeletalMesh.cpp
+++ b/Engine/Plugins/Mutable/Source/MutableTools/Private/MuT/ASTOpAddSkeletalMesh.cpp
@@ -63,22 +63,22 @@ namespace UE::Mutable::Private
 
 	void ASTOpAddSkeletalMesh::Link(FProgram& Program, FLinkerOptions*)
 	{
-		if (linkedAddress)
+		if (LinkedAddress)
 		{
 			return;
 		}
 
-		OP::InstanceAddArgs Args;
+		FOperation::InstanceAddArgs Args;
 		FMemory::Memzero(Args);
 
 		if (Instance)
 		{
-			Args.instance = Instance->linkedAddress;
+			Args.instance = Instance->LinkedAddress;
 		}
 
 		if (SkeletalMesh)
 		{
-			Args.value = SkeletalMesh->linkedAddress;
+			Args.value = SkeletalMesh->LinkedAddress;
 
 			SubtreeRelevantParametersVisitorAST Visitor;
 			Visitor.Run(SkeletalMesh.child());
@@ -103,17 +103,18 @@ namespace UE::Mutable::Private
 
 			if (It != INDEX_NONE)
 			{
-				Program.RelevantParameterList.Add(SkeletalMesh->linkedAddress, It);
+				Program.RelevantParameterList.Add(SkeletalMesh->LinkedAddress, It);
 			}
 			else
 			{
-				Program.RelevantParameterList.Add(SkeletalMesh->linkedAddress, Program.ParameterLists.Num());
+				Program.RelevantParameterList.Add(SkeletalMesh->LinkedAddress, Program.ParameterLists.Num());
 				Program.ParameterLists.Add(Params);
 			}
 		}
 		
-		linkedAddress = static_cast<OP::ADDRESS>(Program.OpAddress.Num());
-		Program.OpAddress.Add(static_cast<uint32_t>(Program.ByteCode.Num()));
+		++Program.NumOps;
+		LinkedAddress = MakeProgramAddress(GetOpType(), Program.ByteCode.Num());
+
 		AppendCode(Program.ByteCode, GetOpType());
 		AppendCode(Program.ByteCode, Args);
 	}
diff --git a/Engine/Plugins/Mutable/Source/MutableTools/Private/MuT/ASTOpBoolAnd.cpp b/Engine/Plugins/Mutable/Source/MutableTools/Private/MuT/ASTOpBoolAnd.cpp
--- a/Engine/Plugins/Mutable/Source/MutableTools/Private/MuT/ASTOpBoolAnd.cpp
+++ b/Engine/Plugins/Mutable/Source/MutableTools/Private/MuT/ASTOpBoolAnd.cpp
@@ -30,8 +30,7 @@ namespace UE::Mutable::Private
 		if (OtherUntyped.GetOpType()==GetOpType())
 		{
 			const ASTOpBoolAnd* Other = static_cast<const ASTOpBoolAnd*>(&OtherUntyped);
-			return A == Other->A &&
-				B == Other->B;
+			return A == Other->A && B == Other->B;
 		}
 		return false;
 	}
@@ -67,16 +66,24 @@ namespace UE::Mutable::Private
 	void ASTOpBoolAnd::Link(FProgram& Program, FLinkerOptions*)
 	{
 		// Already linked?
-		if (!linkedAddress)
+		if (!LinkedAddress)
 		{
-			OP::BoolBinaryArgs Args;
+			FOperation::BoolBinaryArgs Args;
 			FMemory::Memzero(Args);
 
-			if (A) Args.A = A->linkedAddress;
-			if (B) Args.B = B->linkedAddress;
+			if (A) 
+			{
+				Args.A = A->LinkedAddress;
+			}
+
+			if (B)
+			{
+				Args.B = B->LinkedAddress;
+			}
+
+			++Program.NumOps;
+			LinkedAddress = MakeProgramAddress(GetOpType(), Program.ByteCode.Num());
 
-			linkedAddress = (OP::ADDRESS)Program.OpAddress.Num();
-			Program.OpAddress.Add(Program.ByteCode.Num());
 			AppendCode(Program.ByteCode, GetOpType());
 			AppendCode(Program.ByteCode, Args);
 		}
diff --git a/Engine/Plugins/Mutable/Source/MutableTools/Private/MuT/ASTOpBoolEqualIntConst.cpp b/Engine/Plugins/Mutable/Source/MutableTools/Private/MuT/ASTOpBoolEqualIntConst.cpp
--- a/Engine/Plugins/Mutable/Source/MutableTools/Private/MuT/ASTOpBoolEqualIntConst.cpp
+++ b/Engine/Plugins/Mutable/Source/MutableTools/Private/MuT/ASTOpBoolEqualIntConst.cpp
@@ -65,20 +65,21 @@ namespace UE::Mutable::Private
 	void ASTOpBoolEqualIntConst::Link(FProgram& Program, FLinkerOptions*)
 	{
 		// Already linked?
-		if (!linkedAddress)
+		if (!LinkedAddress)
 		{
-			OP::BoolEqualScalarConstArgs Args;
+			FOperation::BoolEqualScalarConstArgs Args;
 			FMemory::Memzero(Args);
 
 			Args.Constant = Constant;
 
 			if (Value)
 			{
-				Args.Value = Value->linkedAddress;
+				Args.Value = Value->LinkedAddress;
 			}
 
-			linkedAddress = (OP::ADDRESS)Program.OpAddress.Num();
-			Program.OpAddress.Add((uint32)Program.ByteCode.Num());
+			++Program.NumOps;
+			LinkedAddress = MakeProgramAddress(GetOpType(), Program.ByteCode.Num());
+
 			AppendCode(Program.ByteCode, GetOpType());
 			AppendCode(Program.ByteCode, Args);
 		}
diff --git a/Engine/Plugins/Mutable/Source/MutableTools/Private/MuT/ASTOpBoolNot.cpp b/Engine/Plugins/Mutable/Source/MutableTools/Private/MuT/ASTOpBoolNot.cpp
--- a/Engine/Plugins/Mutable/Source/MutableTools/Private/MuT/ASTOpBoolNot.cpp
+++ b/Engine/Plugins/Mutable/Source/MutableTools/Private/MuT/ASTOpBoolNot.cpp
@@ -62,15 +62,19 @@ namespace UE::Mutable::Private
 	void ASTOpBoolNot::Link(FProgram& Program, FLinkerOptions*)
 	{
 		// Already linked?
-		if (!linkedAddress)
+		if (!LinkedAddress)
 		{
-			OP::BoolNotArgs Args;
+			FOperation::BoolNotArgs Args;
 			FMemory::Memzero(Args);
 
-			if (A) Args.A = A->linkedAddress;
+			if (A) 
+			{
+				Args.A = A->LinkedAddress;
+			}
+
+			++Program.NumOps;
+			LinkedAddress = MakeProgramAddress(GetOpType(), Program.ByteCode.Num());
 
-			linkedAddress = (OP::ADDRESS)Program.OpAddress.Num();
-			Program.OpAddress.Add(Program.ByteCode.Num());
 			AppendCode(Program.ByteCode, GetOpType());
 			AppendCode(Program.ByteCode, Args);
 		}
diff --git a/Engine/Plugins/Mutable/Source/MutableTools/Private/MuT/ASTOpBoolOr.cpp b/Engine/Plugins/Mutable/Source/MutableTools/Private/MuT/ASTOpBoolOr.cpp
--- a/Engine/Plugins/Mutable/Source/MutableTools/Private/MuT/ASTOpBoolOr.cpp
+++ b/Engine/Plugins/Mutable/Source/MutableTools/Private/MuT/ASTOpBoolOr.cpp
@@ -67,16 +67,24 @@ namespace UE::Mutable::Private
 	void ASTOpBoolOr::Link(FProgram& Program, FLinkerOptions*)
 	{
 		// Already linked?
-		if (!linkedAddress)
+		if (!LinkedAddress)
 		{
-			OP::BoolBinaryArgs Args;
+			FOperation::BoolBinaryArgs Args;
 			FMemory::Memzero(Args);
 
-			if (A) Args.A = A->linkedAddress;
-			if (B) Args.B = B->linkedAddress;
+			if (A) 
+			{
+				Args.A = A->LinkedAddress;
+			}
+
+			if (B) 
+			{
+				Args.B = B->LinkedAddress;
+			}
+
+			++Program.NumOps;
+			LinkedAddress = MakeProgramAddress(GetOpType(), Program.ByteCode.Num());
 
-			linkedAddress = (OP::ADDRESS)Program.OpAddress.Num();
-			Program.OpAddress.Add(Program.ByteCode.Num());
 			AppendCode(Program.ByteCode, GetOpType());
 			AppendCode(Program.ByteCode, Args);
 		}
