[Mutable] Program cache rewrok.

diff --git a/Engine/Plugins/Mutable/Source/MutableTools/Private/MuT/ASTOpColorArithmetic.cpp b/Engine/Plugins/Mutable/Source/MutableTools/Private/MuT/ASTOpColorArithmetic.cpp
--- a/Engine/Plugins/Mutable/Source/MutableTools/Private/MuT/ASTOpColorArithmetic.cpp
+++ b/Engine/Plugins/Mutable/Source/MutableTools/Private/MuT/ASTOpColorArithmetic.cpp
@@ -69,17 +69,26 @@ namespace UE::Mutable::Private
 	void ASTOpColorArithmetic::Link(FProgram& Program, FLinkerOptions*)
 	{
 		// Already linked?
-		if (!linkedAddress)
+		if (!LinkedAddress)
 		{
-			OP::ArithmeticArgs Args;
+			FOperation::ArithmeticArgs Args;
 			FMemory::Memzero(Args);
 
-			if (A) Args.A = A->linkedAddress;
-			if (B) Args.B = B->linkedAddress;
+			if (A) 
+			{
+				Args.A = A->LinkedAddress;
+			}
+
+			if (B)
+			{
+				Args.B = B->LinkedAddress;
+			}
+
 			Args.Operation = Operation;
 
-			linkedAddress = (OP::ADDRESS)Program.OpAddress.Num();
-			Program.OpAddress.Add(Program.ByteCode.Num());
+			++Program.NumOps;
+			LinkedAddress = MakeProgramAddress(GetOpType(), Program.ByteCode.Num());
+
 			AppendCode(Program.ByteCode, GetOpType());
 			AppendCode(Program.ByteCode, Args);
 		}
diff --git a/Engine/Plugins/Mutable/Source/MutableTools/Private/MuT/ASTOpColorFromScalars.cpp b/Engine/Plugins/Mutable/Source/MutableTools/Private/MuT/ASTOpColorFromScalars.cpp
--- a/Engine/Plugins/Mutable/Source/MutableTools/Private/MuT/ASTOpColorFromScalars.cpp
+++ b/Engine/Plugins/Mutable/Source/MutableTools/Private/MuT/ASTOpColorFromScalars.cpp
@@ -90,23 +90,27 @@ namespace UE::Mutable::Private
 	}
 
 
-	void ASTOpColorFromScalars::Link(FProgram& program, FLinkerOptions*)
+	void ASTOpColorFromScalars::Link(FProgram& Program, FLinkerOptions*)
 	{
 		// Already linked?
-		if (!linkedAddress)
+		if (!LinkedAddress)
 		{
-			OP::ColourFromScalarsArgs Args;
+			FOperation::ColourFromScalarsArgs Args;
 			FMemory::Memzero(Args);
 
 			for (int32 i = 0; i < MUTABLE_OP_MAX_SWIZZLE_CHANNELS; ++i)
 			{
-				if (V[i]) Args.V[i] = V[i]->linkedAddress;
+				if (V[i]) 
+				{
+					Args.V[i] = V[i]->LinkedAddress;
+				}
 			}			
 
-			linkedAddress = (OP::ADDRESS)program.OpAddress.Num();
-			program.OpAddress.Add((uint32)program.ByteCode.Num());
-			AppendCode(program.ByteCode, GetOpType());
-			AppendCode(program.ByteCode, Args);
+			++Program.NumOps;
+			LinkedAddress = MakeProgramAddress(GetOpType(), Program.ByteCode.Num());
+
+			AppendCode(Program.ByteCode, GetOpType());
+			AppendCode(Program.ByteCode, Args);
 		}
 	}
 
diff --git a/Engine/Plugins/Mutable/Source/MutableTools/Private/MuT/ASTOpColorSampleImage.cpp b/Engine/Plugins/Mutable/Source/MutableTools/Private/MuT/ASTOpColorSampleImage.cpp
--- a/Engine/Plugins/Mutable/Source/MutableTools/Private/MuT/ASTOpColorSampleImage.cpp
+++ b/Engine/Plugins/Mutable/Source/MutableTools/Private/MuT/ASTOpColorSampleImage.cpp
@@ -74,18 +74,31 @@ namespace UE::Mutable::Private
 	void ASTOpColorSampleImage::Link(FProgram& Program, FLinkerOptions*)
 	{
 		// Already linked?
-		if (!linkedAddress)
+		if (!LinkedAddress)
 		{
-			OP::ColourSampleImageArgs Args;
+			FOperation::ColourSampleImageArgs Args;
 			FMemory::Memzero(Args);
 
-			if (Image) Args.Image = Image->linkedAddress;
-			if (X) Args.X = X->linkedAddress;
-			if (Y) Args.Y = Y->linkedAddress;
+			if (Image) 
+			{
+				Args.Image = Image->LinkedAddress;
+			}
+
+			if (X) 
+			{
+				Args.X = X->LinkedAddress;
+			}
+
+			if (Y) 
+			{
+				Args.Y = Y->LinkedAddress;
+			}
+
 			Args.Filter = Filter;
 
-			linkedAddress = (OP::ADDRESS)Program.OpAddress.Num();
-			Program.OpAddress.Add(Program.ByteCode.Num());
+			++Program.NumOps;
+			LinkedAddress = MakeProgramAddress(GetOpType(), Program.ByteCode.Num());
+
 			AppendCode(Program.ByteCode, GetOpType());
 			AppendCode(Program.ByteCode, Args);
 		}
diff --git a/Engine/Plugins/Mutable/Source/MutableTools/Private/MuT/ASTOpColorSwizzle.cpp b/Engine/Plugins/Mutable/Source/MutableTools/Private/MuT/ASTOpColorSwizzle.cpp
--- a/Engine/Plugins/Mutable/Source/MutableTools/Private/MuT/ASTOpColorSwizzle.cpp
+++ b/Engine/Plugins/Mutable/Source/MutableTools/Private/MuT/ASTOpColorSwizzle.cpp
@@ -96,24 +96,29 @@ namespace UE::Mutable::Private
 	}
 
 
-	void ASTOpColorSwizzle::Link(FProgram& program, FLinkerOptions*)
+	void ASTOpColorSwizzle::Link(FProgram& Program, FLinkerOptions*)
 	{
 		// Already linked?
-		if (!linkedAddress)
+		if (!LinkedAddress)
 		{
-			OP::ColourSwizzleArgs Args;
+			FOperation::ColourSwizzleArgs Args;
 			FMemory::Memzero(Args);
 
 			for (int32 i = 0; i < MUTABLE_OP_MAX_SWIZZLE_CHANNELS; ++i)
 			{
-				if (Sources[i]) Args.sources[i] = Sources[i]->linkedAddress;
+				if (Sources[i]) 
+				{
+					Args.sources[i] = Sources[i]->LinkedAddress;
+				}
+
 				Args.sourceChannels[i] = SourceChannels[i];
 			}			
 
-			linkedAddress = (OP::ADDRESS)program.OpAddress.Num();
-			program.OpAddress.Add((uint32)program.ByteCode.Num());
-			AppendCode(program.ByteCode, GetOpType());
-			AppendCode(program.ByteCode, Args);
+			++Program.NumOps;
+			LinkedAddress = MakeProgramAddress(GetOpType(), Program.ByteCode.Num());	
+
+			AppendCode(Program.ByteCode, GetOpType());
+			AppendCode(Program.ByteCode, Args);
 		}
 	}
 
diff --git a/Engine/Plugins/Mutable/Source/MutableTools/Private/MuT/ASTOpColorToSRGB.cpp b/Engine/Plugins/Mutable/Source/MutableTools/Private/MuT/ASTOpColorToSRGB.cpp
--- a/Engine/Plugins/Mutable/Source/MutableTools/Private/MuT/ASTOpColorToSRGB.cpp
+++ b/Engine/Plugins/Mutable/Source/MutableTools/Private/MuT/ASTOpColorToSRGB.cpp
@@ -61,18 +61,19 @@ namespace UE::Mutable::Private
 	void ASTOpColorToSRGB::Link(FProgram& Program, FLinkerOptions*)
 	{
 		// Already linked?
-		if (!linkedAddress)
+		if (!LinkedAddress)
 		{
-			OP::ColorArgs Args;
+			FOperation::ColorArgs Args;
 			FMemory::Memzero(Args);
 
 			if (Color)
 			{
-				Args.Color = Color->linkedAddress;
+				Args.Color = Color->LinkedAddress;
 			}
 
-			linkedAddress = (OP::ADDRESS)Program.OpAddress.Num();
-			Program.OpAddress.Add(Program.ByteCode.Num());
+			++Program.NumOps;
+			LinkedAddress = MakeProgramAddress(GetOpType(), Program.ByteCode.Num());	
+
 			AppendCode(Program.ByteCode, GetOpType());
 			AppendCode(Program.ByteCode, Args);
 		}
diff --git a/Engine/Plugins/Mutable/Source/MutableTools/Private/MuT/ASTOpComponentAdd.cpp b/Engine/Plugins/Mutable/Source/MutableTools/Private/MuT/ASTOpComponentAdd.cpp
--- a/Engine/Plugins/Mutable/Source/MutableTools/Private/MuT/ASTOpComponentAdd.cpp
+++ b/Engine/Plugins/Mutable/Source/MutableTools/Private/MuT/ASTOpComponentAdd.cpp
@@ -77,24 +77,25 @@ namespace UE::Mutable::Private
 	void ASTOpComponentAdd::Link(FProgram& Program, FLinkerOptions*)
 	{
 		// Already linked?
-		if (!linkedAddress)
+		if (!LinkedAddress)
 		{
-			OP::InstanceAddComponentArgs Args;
+			FOperation::InstanceAddComponentArgs Args;
 			FMemory::Memzero(Args);
 			Args.ExternalId = ExternalId;
 
 			if (Instance)
 			{
-				Args.Instance = Instance->linkedAddress;
+				Args.Instance = Instance->LinkedAddress;
 			}
 				
 			if (Value)
 			{
-				Args.Value = Value->linkedAddress;
+				Args.Value = Value->LinkedAddress;
 			}
 
-			linkedAddress = (OP::ADDRESS)Program.OpAddress.Num();
-			Program.OpAddress.Add((uint32_t)Program.ByteCode.Num());
+			++Program.NumOps;
+			LinkedAddress = MakeProgramAddress(GetOpType(), Program.ByteCode.Num());
+
 			AppendCode(Program.ByteCode, GetOpType());
 			AppendCode(Program.ByteCode, Args);
 		}
diff --git a/Engine/Plugins/Mutable/Source/MutableTools/Private/MuT/ASTOpConditional.cpp b/Engine/Plugins/Mutable/Source/MutableTools/Private/MuT/ASTOpConditional.cpp
--- a/Engine/Plugins/Mutable/Source/MutableTools/Private/MuT/ASTOpConditional.cpp
+++ b/Engine/Plugins/Mutable/Source/MutableTools/Private/MuT/ASTOpConditional.cpp
@@ -107,23 +107,34 @@ namespace UE::Mutable::Private
 
 
 	//-------------------------------------------------------------------------------------------------
-	void ASTOpConditional::Link(FProgram& program, FLinkerOptions*)
+	void ASTOpConditional::Link(FProgram& Program, FLinkerOptions*)
 	{
 		// Already linked?
-		if (!linkedAddress)
+		if (!LinkedAddress)
 		{
-			OP::ConditionalArgs Args;
+			FOperation::ConditionalArgs Args;
 			FMemory::Memzero(Args);
 
-			if (condition) Args.condition = condition->linkedAddress;
-			if (yes) Args.yes = yes->linkedAddress;
-			if (no) Args.no = no->linkedAddress;
+			if (condition) 
+			{
+				Args.condition = condition->LinkedAddress;
+			}
+
+			if (yes) 
+			{
+				Args.yes = yes->LinkedAddress;
+			}
+
+			if (no) 
+			{
+				Args.no = no->LinkedAddress;
+			}
+
+			++Program.NumOps;
+			LinkedAddress = MakeProgramAddress(GetOpType(), Program.ByteCode.Num());	
 
-			linkedAddress = (OP::ADDRESS)program.OpAddress.Num();
-			//program.m_code.push_back(op);
-			program.OpAddress.Add((uint32_t)program.ByteCode.Num());
-			AppendCode(program.ByteCode, type);
-			AppendCode(program.ByteCode, Args);
+			AppendCode(Program.ByteCode, type);
+			AppendCode(Program.ByteCode, Args);
 
 		}
 	}
@@ -399,4 +410,4 @@ namespace UE::Mutable::Private
 		return Result;
 	}
 
-}
\ No newline at end of file
+}
diff --git a/Engine/Plugins/Mutable/Source/MutableTools/Private/MuT/ASTOpConstantBool.cpp b/Engine/Plugins/Mutable/Source/MutableTools/Private/MuT/ASTOpConstantBool.cpp
--- a/Engine/Plugins/Mutable/Source/MutableTools/Private/MuT/ASTOpConstantBool.cpp
+++ b/Engine/Plugins/Mutable/Source/MutableTools/Private/MuT/ASTOpConstantBool.cpp
@@ -53,14 +53,15 @@ namespace UE::Mutable::Private
 
 	void ASTOpConstantBool::Link(FProgram& Program, FLinkerOptions*)
 	{
-		if (!linkedAddress)
+		if (!LinkedAddress)
 		{
-			OP::BoolConstantArgs Args;
+			FOperation::BoolConstantArgs Args;
 			FMemory::Memzero(Args);
 			Args.bValue = bValue;
 
-			linkedAddress = (OP::ADDRESS)Program.OpAddress.Num();
-			Program.OpAddress.Add(Program.ByteCode.Num());
+			++Program.NumOps;
+			LinkedAddress = MakeProgramAddress(GetOpType(), Program.ByteCode.Num());
+
 			AppendCode(Program.ByteCode, GetOpType());
 			AppendCode(Program.ByteCode, Args);
 		}
diff --git a/Engine/Plugins/Mutable/Source/MutableTools/Private/MuT/ASTOpConstantColor.cpp b/Engine/Plugins/Mutable/Source/MutableTools/Private/MuT/ASTOpConstantColor.cpp
--- a/Engine/Plugins/Mutable/Source/MutableTools/Private/MuT/ASTOpConstantColor.cpp
+++ b/Engine/Plugins/Mutable/Source/MutableTools/Private/MuT/ASTOpConstantColor.cpp
@@ -55,14 +55,15 @@ namespace UE::Mutable::Private
 
 	void ASTOpConstantColor::Link(FProgram& Program, FLinkerOptions*)
 	{
-		if (!linkedAddress)
+		if (!LinkedAddress)
 		{
-			OP::ColorConstantArgs Args;
+			FOperation::ColorConstantArgs Args;
 			FMemory::Memzero(Args);
 			Args.Value = Value;
 
-			linkedAddress = (OP::ADDRESS)Program.OpAddress.Num();
-			Program.OpAddress.Add(Program.ByteCode.Num());
+			++Program.NumOps;
+			LinkedAddress = MakeProgramAddress(GetOpType(), Program.ByteCode.Num());
+
 			AppendCode(Program.ByteCode, GetOpType());
 			AppendCode(Program.ByteCode, Args);
 		}
diff --git a/Engine/Plugins/Mutable/Source/MutableTools/Private/MuT/ASTOpConstantExtensionData.cpp b/Engine/Plugins/Mutable/Source/MutableTools/Private/MuT/ASTOpConstantExtensionData.cpp
--- a/Engine/Plugins/Mutable/Source/MutableTools/Private/MuT/ASTOpConstantExtensionData.cpp
+++ b/Engine/Plugins/Mutable/Source/MutableTools/Private/MuT/ASTOpConstantExtensionData.cpp
@@ -47,18 +47,19 @@ uint32 ASTOpConstantExtensionData::Hash() const
 //-------------------------------------------------------------------------------------------------
 void ASTOpConstantExtensionData::Link(FProgram& Program, FLinkerOptions* Options)
 {
-	if (!linkedAddress)
+	if (!LinkedAddress)
 	{
-		OP::ExternalDataConstantArgs Args;
+		FOperation::ExternalDataConstantArgs Args;
 		FMemory::Memzero(Args);
 
 		Args.ExternalObjectId = ExtensionDataId;
 
-		linkedAddress = (OP::ADDRESS)Program.OpAddress.Num();
-		Program.OpAddress.Add((uint32_t)Program.ByteCode.Num());
+		++Program.NumOps;
+		LinkedAddress = MakeProgramAddress(GetOpType(), Program.ByteCode.Num());
+
 		AppendCode(Program.ByteCode, GetOpType());
 		AppendCode(Program.ByteCode, Args);
 	}
 }
 
-}
\ No newline at end of file
+}
diff --git a/Engine/Plugins/Mutable/Source/MutableTools/Private/MuT/ASTOpConstantInt.cpp b/Engine/Plugins/Mutable/Source/MutableTools/Private/MuT/ASTOpConstantInt.cpp
--- a/Engine/Plugins/Mutable/Source/MutableTools/Private/MuT/ASTOpConstantInt.cpp
+++ b/Engine/Plugins/Mutable/Source/MutableTools/Private/MuT/ASTOpConstantInt.cpp
@@ -53,14 +53,15 @@ namespace UE::Mutable::Private
 
 	void ASTOpConstantInt::Link(FProgram& Program, FLinkerOptions*)
 	{
-		if (!linkedAddress)
+		if (!LinkedAddress)
 		{
-			OP::IntConstantArgs Args;
+			FOperation::IntConstantArgs Args;
 			FMemory::Memzero(Args);
 			Args.Value = Value;
 
-			linkedAddress = (OP::ADDRESS)Program.OpAddress.Num();
-			Program.OpAddress.Add(Program.ByteCode.Num());
+			++Program.NumOps;
+			LinkedAddress = MakeProgramAddress(GetOpType(), Program.ByteCode.Num());
+			
 			AppendCode(Program.ByteCode, GetOpType());
 			AppendCode(Program.ByteCode, Args);
 		}
diff --git a/Engine/Plugins/Mutable/Source/MutableTools/Private/MuT/ASTOpConstantMatrix.cpp b/Engine/Plugins/Mutable/Source/MutableTools/Private/MuT/ASTOpConstantMatrix.cpp
--- a/Engine/Plugins/Mutable/Source/MutableTools/Private/MuT/ASTOpConstantMatrix.cpp
+++ b/Engine/Plugins/Mutable/Source/MutableTools/Private/MuT/ASTOpConstantMatrix.cpp
@@ -33,17 +33,18 @@ UE::Mutable::Private::Ptr<UE::Mutable::Private::ASTOp> UE::Mutable::Private::AST
 	return n;
 }
 
-void UE::Mutable::Private::ASTOpConstantMatrix::Link(FProgram& program, FLinkerOptions* Options)
+void UE::Mutable::Private::ASTOpConstantMatrix::Link(FProgram& Program, FLinkerOptions* Options)
 {
-	if (!linkedAddress)
+	if (!LinkedAddress)
 	{
-		OP::MatrixConstantArgs Args;
+		FOperation::MatrixConstantArgs Args;
 		FMemory::Memzero(Args);
-		Args.value = program.AddConstant(value);
+		Args.value = Program.AddConstant(value);
 
-		linkedAddress = static_cast<OP::ADDRESS>(program.OpAddress.Num());
-		program.OpAddress.Add(static_cast<uint32_t>(program.ByteCode.Num()));
-		AppendCode(program.ByteCode, EOpType::MA_CONSTANT);
-		AppendCode(program.ByteCode, Args);
+		++Program.NumOps;
+		LinkedAddress = MakeProgramAddress(GetOpType(), Program.ByteCode.Num());
+		
+		AppendCode(Program.ByteCode, GetOpType());
+		AppendCode(Program.ByteCode, Args);
 	}
 }
diff --git a/Engine/Plugins/Mutable/Source/MutableTools/Private/MuT/ASTOpConstantProjector.cpp b/Engine/Plugins/Mutable/Source/MutableTools/Private/MuT/ASTOpConstantProjector.cpp
--- a/Engine/Plugins/Mutable/Source/MutableTools/Private/MuT/ASTOpConstantProjector.cpp
+++ b/Engine/Plugins/Mutable/Source/MutableTools/Private/MuT/ASTOpConstantProjector.cpp
@@ -49,19 +49,19 @@ namespace UE::Mutable::Private
 	}
 
 
-	void ASTOpConstantProjector::Link(FProgram& program, FLinkerOptions*)
+	void ASTOpConstantProjector::Link(FProgram& Program, FLinkerOptions*)
 	{
-		if (!linkedAddress)
+		if (!LinkedAddress)
 		{
-			OP::ResourceConstantArgs Args;
+			FOperation::ResourceConstantArgs Args;
 			FMemory::Memzero(Args);
-			Args.value = program.AddConstant(value);
+			Args.value = Program.AddConstant(value);
 
-			linkedAddress = (OP::ADDRESS)program.OpAddress.Num();
-			//program.m_code.push_back(op);
-			program.OpAddress.Add((uint32_t)program.ByteCode.Num());
-			AppendCode(program.ByteCode, EOpType::PR_CONSTANT);
-			AppendCode(program.ByteCode, Args);
+			++Program.NumOps;
+			LinkedAddress = MakeProgramAddress(GetOpType(), Program.ByteCode.Num());
+			
+			AppendCode(Program.ByteCode, GetOpType());
+			AppendCode(Program.ByteCode, Args);
 		}
 	}
 
diff --git a/Engine/Plugins/Mutable/Source/MutableTools/Private/MuT/ASTOpConstantResource.cpp b/Engine/Plugins/Mutable/Source/MutableTools/Private/MuT/ASTOpConstantResource.cpp
--- a/Engine/Plugins/Mutable/Source/MutableTools/Private/MuT/ASTOpConstantResource.cpp
+++ b/Engine/Plugins/Mutable/Source/MutableTools/Private/MuT/ASTOpConstantResource.cpp
@@ -77,7 +77,7 @@ namespace UE::Mutable::Private
 
 	namespace
 	{
- 		/** Adds a constant mesh data to a program and returns its constant index. */
+ 		/** Adds a constant mesh data to a Program and returns its constant index. */
 		int32 AddConstantMesh(FProgram& Program, const TManagedPtr<const FMesh>& MeshData, FLinkerOptions& Options)
 		{
 			auto AddMeshToProgram = [&Options, &Program](const TManagedPtr<FMesh>& Mesh)
@@ -291,7 +291,7 @@ namespace UE::Mutable::Private
 			return Program.ConstantMeshes.Add(MeshContentRange);
 		}   
 
-		/** Adds a constant image data to a program and returns its constant index. */
+		/** Adds a constant image data to a Program and returns its constant index. */
 		uint32 AddConstantImage(FProgram& Program, const TManagedPtr<const FImage>& Image, FLinkerOptions& Options)
 		{
 			MUTABLE_CPUPROFILER_SCOPE(AddConstantImage);
@@ -432,126 +432,129 @@ namespace UE::Mutable::Private
 		}
 	}
 
-	//-------------------------------------------------------------------------------------------------
 	void ASTOpConstantResource::Link(FProgram& Program, FLinkerOptions* Options)
 	{
 		MUTABLE_CPUPROFILER_SCOPE(ASTOpConstantResource_Link);
 
-		if (!linkedAddress && !bLinkedAndNull)
+		if (!LinkedAddress && !bLinkedAndNull)
 		{
 			switch (Type)
 			{
 			case EOpType::ME_CONSTANT:
-				{
-					OP::MeshConstantArgs Args;
-					FMemory::Memzero(Args);
+			{
+				FOperation::MeshConstantArgs Args;
+				FMemory::Memzero(Args);
 
-					TManagedPtr<const FMesh> MeshValue = StaticCastManagedPtr<const FMesh>(Value);
-					check(MeshValue);
+				TManagedPtr<const FMesh> MeshValue = StaticCastManagedPtr<const FMesh>(Value);
+				check(MeshValue);
 
-					Args.Skeleton = -1;
-					if (TManagedPtr<const FSkeleton> MeshSkeleton = MeshValue->GetSkeleton())
-					{
-						Args.Skeleton = Program.AddConstant(MeshSkeleton);
-						//MeshValue->Skeleton = nullptr;
-					}
+				Args.Skeleton = -1;
+				if (TManagedPtr<const FSkeleton> MeshSkeleton = MeshValue->GetSkeleton())
+				{
+					Args.Skeleton = Program.AddConstant(MeshSkeleton);
+					//MeshValue->Skeleton = nullptr;
+				}
 
-					check(!MeshValue->PassthroughObject);
-					
-					Args.Value = AddConstantMesh(Program, MeshValue, *Options);
-					Args.ClothID = MeshValue->ClothSections.Num() ? MeshValue->ClothSections[0].ClothingAsset.GetId() : PASSTHROUGH_ID_INVALID;
+				check(!MeshValue->PassthroughObject);
+				
+				Args.Value = AddConstantMesh(Program, MeshValue, *Options);
+				Args.ClothID = MeshValue->ClothSections.Num() ? MeshValue->ClothSections[0].ClothingAsset.GetId() : PASSTHROUGH_ID_INVALID;
 
-					int32 DataDescIndex = Options->AdditionalData.SourceMeshPerConstant.Add(SourceDataDescriptor);
-					check(DataDescIndex == Args.Value);
+				int32 DataDescIndex = Options->AdditionalData.SourceMeshPerConstant.Add(SourceDataDescriptor);
+				check(DataDescIndex == Args.Value);
 
-					linkedAddress = (OP::ADDRESS)Program.OpAddress.Num();
-					Program.OpAddress.Add((uint32)Program.ByteCode.Num());
-					AppendCode(Program.ByteCode, Type);
-					AppendCode(Program.ByteCode, Args);
+				++Program.NumOps;
+				LinkedAddress = MakeProgramAddress(GetOpType(), Program.ByteCode.Num());
+				
+				AppendCode(Program.ByteCode, Type);
+				AppendCode(Program.ByteCode, Args);
 
-					break;
-				}
+				break;
+			}
 
 			case EOpType::IM_CONSTANT:
-				{
-					OP::ResourceConstantArgs Args;
-					FMemory::Memzero(Args);
-					
-					TManagedPtr<const FImage> pTyped = StaticCastManagedPtr<const FImage>(Value);
-					check(pTyped);
-
-					if (pTyped->GetSizeX() * pTyped->GetSizeY() == 0)
-					{
-						// It's an empty or degenerated image, return a null operation.
-						// Null op
-						linkedAddress = 0;
-						bLinkedAndNull = true;
-					}
-					else
-					{
-						check(!pTyped->PassthroughObject);
-
-						Args.value = AddConstantImage( Program, pTyped, *Options);
-						
-						int32 DataDescIndex = Options->AdditionalData.SourceImagePerConstant.Add(SourceDataDescriptor);
-						check(DataDescIndex == Args.value);
-						
-						linkedAddress = (OP::ADDRESS)Program.OpAddress.Num();
-						Program.OpAddress.Add((uint32)Program.ByteCode.Num());
-						AppendCode(Program.ByteCode, Type);
-						AppendCode(Program.ByteCode, Args);
-					}
+			{
+				FOperation::ResourceConstantArgs Args;
+				FMemory::Memzero(Args);
+				
+				TManagedPtr<const FImage> pTyped = StaticCastManagedPtr<const FImage>(Value);
+				check(pTyped);
 
-					break;
+				if (pTyped->GetSizeX() * pTyped->GetSizeY() == 0)
+				{
+					// It's an empty or degenerated image, return a null operation.
+					// Null op
+					LinkedAddress = 0;
+					bLinkedAndNull = true;
 				}
-				
-			case EOpType::LA_CONSTANT:
+				else
 				{
-					OP::ResourceConstantArgs Args;
-					FMemory::Memzero(Args);
-					
-					TManagedPtr<const FLayout> pTyped = StaticCastManagedPtr<const FLayout>(Value);
-					check(pTyped);
-					Args.value = Program.AddConstant(pTyped);
+					check(!pTyped->PassthroughObject);
 
-					linkedAddress = (OP::ADDRESS)Program.OpAddress.Num();
-					Program.OpAddress.Add((uint32)Program.ByteCode.Num());
+					Args.value = AddConstantImage( Program, pTyped, *Options);
+					
+					int32 DataDescIndex = Options->AdditionalData.SourceImagePerConstant.Add(SourceDataDescriptor);
+					check(DataDescIndex == Args.value);
+					
+					++Program.NumOps;
+					LinkedAddress = MakeProgramAddress(GetOpType(), Program.ByteCode.Num());
+					
 					AppendCode(Program.ByteCode, Type);
 					AppendCode(Program.ByteCode, Args);
-					
-					break;
 				}
 
-			case EOpType::MI_CONSTANT:
-				{
-					OP::MaterialConstantArgs Args;
-					FMemory::Memzero(Args);
+				break;
+			}
+				
+			case EOpType::LA_CONSTANT:
+			{
+				FOperation::ResourceConstantArgs Args;
+				FMemory::Memzero(Args);
+				
+				TManagedPtr<const FLayout> pTyped = StaticCastManagedPtr<const FLayout>(Value);
+				check(pTyped);
+				Args.value = Program.AddConstant(pTyped);
 
-					TManagedPtr<FMaterial> MaterialValue = StaticCastManagedPtr<const FMaterial>(Value)->Clone();
-					check(MaterialValue);
+				++Program.NumOps;
+				LinkedAddress = MakeProgramAddress(GetOpType(), Program.ByteCode.Num());
+				
+				AppendCode(Program.ByteCode, Type);
+				AppendCode(Program.ByteCode, Args);
+				
+				break;
+			}
 
-					// Store linked child images into the material constant
-					for (const TPair<FParameterKey, ASTChild>& ImageOperation : ImageOperations)
-					{
-						TVariant<OP::ADDRESS, TManagedPtr<const FImage>> NewImage;
-						NewImage.Set<OP::ADDRESS>(ImageOperation.Value->linkedAddress);
+			case EOpType::MI_CONSTANT:
+			{
+				FOperation::MaterialConstantArgs Args;
+				FMemory::Memzero(Args);
 
-						FMaterial::FImageParameterData ImageParameterData;
-						ImageParameterData.ImageParameter = NewImage;
+				TManagedPtr<FMaterial> MaterialValue = StaticCastManagedPtr<const FMaterial>(Value)->Clone();
+				check(MaterialValue);
 
-						MaterialValue->ImageParameters.Add(ImageOperation.Key, ImageParameterData);
-					}
+				// Store linked child images into the material constant
+				for (const TPair<FParameterKey, ASTChild>& ImageOperation : ImageOperations)
+				{
+					TVariant<FOperation::ADDRESS, TManagedPtr<const FImage>> NewImage;
+					NewImage.Set<FOperation::ADDRESS>(ImageOperation.Value->LinkedAddress);
 
-					Args.ID = MaterialValue->PassthroughObject.Reset(); // Not serialized but compared. Set to invalid to avoid duplicating it.
-					Args.Value = Program.AddConstant(MaterialValue);
-					
-					linkedAddress = (OP::ADDRESS)Program.OpAddress.Num();
-					Program.OpAddress.Add((uint32)Program.ByteCode.Num());
-					AppendCode(Program.ByteCode, Type);
-					AppendCode(Program.ByteCode, Args);
-					
-					break;
+					FMaterial::FImageParameterData ImageParameterData;
+					ImageParameterData.ImageParameter = NewImage;
+
+					MaterialValue->ImageParameters.Add(ImageOperation.Key, ImageParameterData);
 				}
+
+				Args.ID = MaterialValue->PassthroughObject.Reset(); // Not serialized but compared. Set to invalid to avoid duplicating it.
+				Args.Value = Program.AddConstant(MaterialValue);
+				
+				++Program.NumOps;
+				LinkedAddress = MakeProgramAddress(GetOpType(), Program.ByteCode.Num());
+				
+				AppendCode(Program.ByteCode, GetOpType());
+				AppendCode(Program.ByteCode, Args);
+				
+				break;
+			}
 			default:
 				unimplemented();
 			}
@@ -561,8 +564,6 @@ namespace UE::Mutable::Private
 		}
 	}
 
-
-	//-------------------------------------------------------------------------------------------------
 	FImageDesc ASTOpConstantResource::GetImageDesc(bool, class FGetImageDescContext* Context) const
 	{
 		FImageDesc Result;
diff --git a/Engine/Plugins/Mutable/Source/MutableTools/Private/MuT/ASTOpConstantScalar.cpp b/Engine/Plugins/Mutable/Source/MutableTools/Private/MuT/ASTOpConstantScalar.cpp
--- a/Engine/Plugins/Mutable/Source/MutableTools/Private/MuT/ASTOpConstantScalar.cpp
+++ b/Engine/Plugins/Mutable/Source/MutableTools/Private/MuT/ASTOpConstantScalar.cpp
@@ -53,14 +53,15 @@ namespace UE::Mutable::Private
 
 	void ASTOpConstantScalar::Link(FProgram& Program, FLinkerOptions*)
 	{
-		if (!linkedAddress)
+		if (!LinkedAddress)
 		{
-			OP::ScalarConstantArgs Args;
+			FOperation::ScalarConstantArgs Args;
 			FMemory::Memzero(Args);
 			Args.Value = Value;
 
-			linkedAddress = (OP::ADDRESS)Program.OpAddress.Num();
-			Program.OpAddress.Add(Program.ByteCode.Num());
+			++Program.NumOps;
+			LinkedAddress = MakeProgramAddress(GetOpType(), Program.ByteCode.Num());
+			
 			AppendCode(Program.ByteCode, GetOpType());
 			AppendCode(Program.ByteCode, Args);
 		}
diff --git a/Engine/Plugins/Mutable/Source/MutableTools/Private/MuT/ASTOpConstantString.cpp b/Engine/Plugins/Mutable/Source/MutableTools/Private/MuT/ASTOpConstantString.cpp
--- a/Engine/Plugins/Mutable/Source/MutableTools/Private/MuT/ASTOpConstantString.cpp
+++ b/Engine/Plugins/Mutable/Source/MutableTools/Private/MuT/ASTOpConstantString.cpp
@@ -49,19 +49,20 @@ namespace UE::Mutable::Private
 	}
 
 
-	void ASTOpConstantString::Link(FProgram& program, FLinkerOptions*)
+	void ASTOpConstantString::Link(FProgram& Program, FLinkerOptions*)
 	{
-		if (!linkedAddress)
+		if (!LinkedAddress)
 		{
-			OP::ResourceConstantArgs Args;
+			FOperation::ResourceConstantArgs Args;
 			FMemory::Memzero(Args);
-			Args.value = program.AddConstant(value);
+			Args.value = Program.AddConstant(value);
 
-			linkedAddress = (OP::ADDRESS)program.OpAddress.Num();
-			program.OpAddress.Add((uint32_t)program.ByteCode.Num());
-			AppendCode(program.ByteCode, EOpType::ST_CONSTANT);
-			AppendCode(program.ByteCode, Args);
+			++Program.NumOps;
+			LinkedAddress = MakeProgramAddress(GetOpType(), Program.ByteCode.Num());
+			
+			AppendCode(Program.ByteCode, EOpType::ST_CONSTANT);
+			AppendCode(Program.ByteCode, Args);
 		}
 	}
 
-}
\ No newline at end of file
+}
diff --git a/Engine/Plugins/Mutable/Source/MutableTools/Private/MuT/ASTOpEntry.cpp b/Engine/Plugins/Mutable/Source/MutableTools/Private/MuT/ASTOpEntry.cpp
--- a/Engine/Plugins/Mutable/Source/MutableTools/Private/MuT/ASTOpEntry.cpp
+++ b/Engine/Plugins/Mutable/Source/MutableTools/Private/MuT/ASTOpEntry.cpp
@@ -66,11 +66,11 @@ namespace UE::Mutable::Private
 
 	void ASTOpEntry::Link(FProgram& Program, FLinkerOptions*)
 	{
-		if (!linkedAddress)
+		if (!LinkedAddress)
 		{
 			if (A)
 			{
-				linkedAddress = A->linkedAddress;
+				LinkedAddress = A->LinkedAddress;
 			}
 		}
 	}
diff --git a/Engine/Plugins/Mutable/Source/MutableTools/Private/MuT/ASTOpExternal.cpp b/Engine/Plugins/Mutable/Source/MutableTools/Private/MuT/ASTOpExternal.cpp
--- a/Engine/Plugins/Mutable/Source/MutableTools/Private/MuT/ASTOpExternal.cpp
+++ b/Engine/Plugins/Mutable/Source/MutableTools/Private/MuT/ASTOpExternal.cpp
@@ -81,27 +81,28 @@ namespace UE::Mutable::Private
 
 	void ASTOpExternal::Link(FProgram& Program, FLinkerOptions* Options)
 	{
-		if (!linkedAddress)
+		if (!LinkedAddress)
 		{
-			OP::ExternalArgs Args;
+			FOperation::ExternalArgs Args;
 			FMemory::Memzero(Args);
 			Args.NumOperants = Inputs.Num();
 
-			linkedAddress = (OP::ADDRESS)Program.OpAddress.Num();
-			Program.OpAddress.Add((uint32_t)Program.ByteCode.Num());
-			
 			check(Type != EOpType::NONE);
+
+			++Program.NumOps;
+			LinkedAddress = MakeProgramAddress(Type, Program.ByteCode.Num());
+			
 			AppendCode(Program.ByteCode, Type);
 			AppendCode(Program.ByteCode, Args);
 
 			for (const ASTChild& Input : Inputs)
 			{
-				AppendCode(Program.ByteCode, Input->linkedAddress);
+				AppendCode(Program.ByteCode, Input->LinkedAddress);
 			}
 
 			if (Options->ExternalOperations)
 			{
-				Options->ExternalOperations->Add(linkedAddress, OperationInstancedStruct);
+				Options->ExternalOperations->Add(LinkedAddress, OperationInstancedStruct);
 			}
 		}
 	}
diff --git a/Engine/Plugins/Mutable/Source/MutableTools/Private/MuT/ASTOpImageBinarize.cpp b/Engine/Plugins/Mutable/Source/MutableTools/Private/MuT/ASTOpImageBinarize.cpp
--- a/Engine/Plugins/Mutable/Source/MutableTools/Private/MuT/ASTOpImageBinarize.cpp
+++ b/Engine/Plugins/Mutable/Source/MutableTools/Private/MuT/ASTOpImageBinarize.cpp
@@ -80,22 +80,23 @@ namespace UE::Mutable::Private
 	void ASTOpImageBinarize::Link(FProgram& Program, FLinkerOptions*)
 	{
 		// Already linked?
-		if (!linkedAddress)
+		if (!LinkedAddress)
 		{
-			OP::ImageBinariseArgs Args;
+			FOperation::ImageBinariseArgs Args;
 			FMemory::Memzero(Args);
 
 			if (Base)
 			{
-				Args.Base = Base->linkedAddress;
+				Args.Base = Base->LinkedAddress;
 			}
 			if (Threshold)
 			{
-				Args.Threshold = Threshold->linkedAddress;
+				Args.Threshold = Threshold->LinkedAddress;
 			}
 
-			linkedAddress = (OP::ADDRESS)Program.OpAddress.Num();
-			Program.OpAddress.Add(Program.ByteCode.Num());
+			++Program.NumOps;
+			LinkedAddress = MakeProgramAddress(GetOpType(), Program.ByteCode.Num());
+			
 			AppendCode(Program.ByteCode, GetOpType());
 			AppendCode(Program.ByteCode, Args);
 		}
diff --git a/Engine/Plugins/Mutable/Source/MutableTools/Private/MuT/ASTOpImageBlankLayout.cpp b/Engine/Plugins/Mutable/Source/MutableTools/Private/MuT/ASTOpImageBlankLayout.cpp
--- a/Engine/Plugins/Mutable/Source/MutableTools/Private/MuT/ASTOpImageBlankLayout.cpp
+++ b/Engine/Plugins/Mutable/Source/MutableTools/Private/MuT/ASTOpImageBlankLayout.cpp
@@ -82,14 +82,14 @@ namespace UE::Mutable::Private
 	void ASTOpImageBlankLayout::Link(FProgram& Program, FLinkerOptions*)
 	{
 		// Already linked?
-		if (!linkedAddress)
+		if (!LinkedAddress)
 		{
-			OP::ImageBlankLayoutArgs Args;
+			FOperation::ImageBlankLayoutArgs Args;
 			FMemory::Memzero(Args);
 
 			if (Layout)
 			{
-				Args.Layout = Layout->linkedAddress;
+				Args.Layout = Layout->LinkedAddress;
 			}
 
 			Args.BlockSize[0] = BlockSize[0];
@@ -98,8 +98,9 @@ namespace UE::Mutable::Private
 			Args.GenerateMipmaps = GenerateMipmaps;
 			Args.MipmapCount = MipmapCount;
 
-			linkedAddress = (OP::ADDRESS)Program.OpAddress.Num();
-			Program.OpAddress.Add(Program.ByteCode.Num());
+			++Program.NumOps;
+			LinkedAddress = MakeProgramAddress(GetOpType(), Program.ByteCode.Num());
+			
 			AppendCode(Program.ByteCode, GetOpType());
 			AppendCode(Program.ByteCode, Args);
 		}
diff --git a/Engine/Plugins/Mutable/Source/MutableTools/Private/MuT/ASTOpImageColorMap.cpp b/Engine/Plugins/Mutable/Source/MutableTools/Private/MuT/ASTOpImageColorMap.cpp
--- a/Engine/Plugins/Mutable/Source/MutableTools/Private/MuT/ASTOpImageColorMap.cpp
+++ b/Engine/Plugins/Mutable/Source/MutableTools/Private/MuT/ASTOpImageColorMap.cpp
@@ -85,26 +85,27 @@ namespace UE::Mutable::Private
 	void ASTOpImageColorMap::Link(FProgram& Program, FLinkerOptions*)
 	{
 		// Already linked?
-		if (!linkedAddress)
+		if (!LinkedAddress)
 		{
-			OP::ImageColourMapArgs Args;
+			FOperation::ImageColourMapArgs Args;
 			FMemory::Memzero(Args);
 
 			if (Base)
 			{
-				Args.Base = Base->linkedAddress;
+				Args.Base = Base->LinkedAddress;
 			}
 			if (Mask)
 			{
-				Args.Mask = Mask->linkedAddress;
+				Args.Mask = Mask->LinkedAddress;
 			}
 			if (Map)
 			{
-				Args.Map = Map->linkedAddress;
+				Args.Map = Map->LinkedAddress;
 			}
 
-			linkedAddress = (OP::ADDRESS)Program.OpAddress.Num();
-			Program.OpAddress.Add(Program.ByteCode.Num());
+			++Program.NumOps;
+			LinkedAddress = MakeProgramAddress(GetOpType(), Program.ByteCode.Num());
+			
 			AppendCode(Program.ByteCode, GetOpType());
 			AppendCode(Program.ByteCode, Args);
 		}
diff --git a/Engine/Plugins/Mutable/Source/MutableTools/Private/MuT/ASTOpImageCompose.cpp b/Engine/Plugins/Mutable/Source/MutableTools/Private/MuT/ASTOpImageCompose.cpp
--- a/Engine/Plugins/Mutable/Source/MutableTools/Private/MuT/ASTOpImageCompose.cpp
+++ b/Engine/Plugins/Mutable/Source/MutableTools/Private/MuT/ASTOpImageCompose.cpp
@@ -96,24 +96,41 @@ void ASTOpImageCompose::ForEachChild(const TFunctionRef<void(ASTChild&)> f )
 
 
 //-------------------------------------------------------------------------------------------------
-void ASTOpImageCompose::Link( FProgram& program, FLinkerOptions*)
+void ASTOpImageCompose::Link( FProgram& Program, FLinkerOptions*)
 {
     // Already linked?
-    if (!linkedAddress)
+    if (!LinkedAddress)
     {
-        OP::ImageComposeArgs Args;
+        FOperation::ImageComposeArgs Args;
 		FMemory::Memzero(Args);
 
-		if (Layout) Args.layout = Layout->linkedAddress;
-		if (Base) Args.base = Base->linkedAddress;
-		if (BlockImage) Args.blockImage = BlockImage->linkedAddress;
-		if (Mask) Args.mask = Mask->linkedAddress;
+		if (Layout) 
+		{
+			Args.layout = Layout->LinkedAddress;
+		}
+
+		if (Base) 
+		{
+			Args.base = Base->LinkedAddress;
+		}
+
+		if (BlockImage) 
+		{
+			Args.blockImage = BlockImage->LinkedAddress;
+		}
+
+		if (Mask) 
+		{
+			Args.mask = Mask->LinkedAddress;
+		}
+
         Args.BlockId = BlockId;
 
-        linkedAddress = (OP::ADDRESS)program.OpAddress.Num();
-        program.OpAddress.Add((uint32_t)program.ByteCode.Num());
-        AppendCode(program.ByteCode,EOpType::IM_COMPOSE);
-        AppendCode(program.ByteCode,Args);
+        ++Program.NumOps;
+        LinkedAddress = MakeProgramAddress(GetOpType(), Program.ByteCode.Num());
+        
+        AppendCode(Program.ByteCode, GetOpType());
+        AppendCode(Program.ByteCode,Args);
     }
 
 }
@@ -339,4 +356,4 @@ FSourceDataDescriptor ASTOpImageCompose::GetSourceDataDescriptor(FGetSourceDataD
 }
 
 
-}
\ No newline at end of file
+}
diff --git a/Engine/Plugins/Mutable/Source/MutableTools/Private/MuT/ASTOpImageCrop.cpp b/Engine/Plugins/Mutable/Source/MutableTools/Private/MuT/ASTOpImageCrop.cpp
--- a/Engine/Plugins/Mutable/Source/MutableTools/Private/MuT/ASTOpImageCrop.cpp
+++ b/Engine/Plugins/Mutable/Source/MutableTools/Private/MuT/ASTOpImageCrop.cpp
@@ -78,14 +78,14 @@ namespace UE::Mutable::Private
 	void ASTOpImageCrop::Link(FProgram& Program, FLinkerOptions*)
 	{
 		// Already linked?
-		if (!linkedAddress)
+		if (!LinkedAddress)
 		{
-			OP::ImageCropArgs Args;
+			FOperation::ImageCropArgs Args;
 			FMemory::Memzero(Args);
 
 			if (Source)
 			{
-				Args.source = Source->linkedAddress;
+				Args.source = Source->LinkedAddress;
 			}
 
 			Args.minX = Min[0];
@@ -93,8 +93,9 @@ namespace UE::Mutable::Private
 			Args.sizeX = Size[0];
 			Args.sizeY = Size[1];
 
-			linkedAddress = (OP::ADDRESS)Program.OpAddress.Num();
-			Program.OpAddress.Add(Program.ByteCode.Num());
+			++Program.NumOps;
+			LinkedAddress = MakeProgramAddress(GetOpType(), Program.ByteCode.Num());
+
 			AppendCode(Program.ByteCode, GetOpType());
 			AppendCode(Program.ByteCode, Args);
 		}
diff --git a/Engine/Plugins/Mutable/Source/MutableTools/Private/MuT/ASTOpImageDisplace.cpp b/Engine/Plugins/Mutable/Source/MutableTools/Private/MuT/ASTOpImageDisplace.cpp
--- a/Engine/Plugins/Mutable/Source/MutableTools/Private/MuT/ASTOpImageDisplace.cpp
+++ b/Engine/Plugins/Mutable/Source/MutableTools/Private/MuT/ASTOpImageDisplace.cpp
@@ -80,22 +80,23 @@ namespace UE::Mutable::Private
 	void ASTOpImageDisplace::Link(FProgram& Program, FLinkerOptions*)
 	{
 		// Already linked?
-		if (!linkedAddress)
+		if (!LinkedAddress)
 		{
-			OP::ImageDisplaceArgs Args;
+			FOperation::ImageDisplaceArgs Args;
 			FMemory::Memzero(Args);
 
 			if (Source)
 			{
-				Args.Source = Source->linkedAddress;
+				Args.Source = Source->LinkedAddress;
 			}
 			if (DisplacementMap)
 			{
-				Args.DisplacementMap = DisplacementMap->linkedAddress;
+				Args.DisplacementMap = DisplacementMap->LinkedAddress;
 			}
 
-			linkedAddress = (OP::ADDRESS)Program.OpAddress.Num();
-			Program.OpAddress.Add(Program.ByteCode.Num());
+			++Program.NumOps;
+			LinkedAddress = MakeProgramAddress(GetOpType(), Program.ByteCode.Num());
+			
 			AppendCode(Program.ByteCode, GetOpType());
 			AppendCode(Program.ByteCode, Args);
 		}
diff --git a/Engine/Plugins/Mutable/Source/MutableTools/Private/MuT/ASTOpImageInterpolate.cpp b/Engine/Plugins/Mutable/Source/MutableTools/Private/MuT/ASTOpImageInterpolate.cpp
--- a/Engine/Plugins/Mutable/Source/MutableTools/Private/MuT/ASTOpImageInterpolate.cpp
+++ b/Engine/Plugins/Mutable/Source/MutableTools/Private/MuT/ASTOpImageInterpolate.cpp
@@ -96,27 +96,32 @@ namespace UE::Mutable::Private
 	}
 
 
-	void ASTOpImageInterpolate::Link(FProgram& program, FLinkerOptions*)
+	void ASTOpImageInterpolate::Link(FProgram& Program, FLinkerOptions*)
 	{
 		// Already linked?
-		if (!linkedAddress)
+		if (!LinkedAddress)
 		{
-			OP::ImageInterpolateArgs Args;
+			FOperation::ImageInterpolateArgs Args;
 			FMemory::Memzero(Args);
 
 			if (Factor)
 			{
-				Args.Factor = Factor->linkedAddress;
+				Args.Factor = Factor->LinkedAddress;
 			}
+
 			for (int32 i = 0; i < MUTABLE_OP_MAX_INTERPOLATE_COUNT; ++i)
 			{
-				if (Targets[i]) Args.Targets[i] = Targets[i]->linkedAddress;
+				if (Targets[i]) 
+				{
+					Args.Targets[i] = Targets[i]->LinkedAddress;
+				}
 			}			
 
-			linkedAddress = (OP::ADDRESS)program.OpAddress.Num();
-			program.OpAddress.Add((uint32)program.ByteCode.Num());
-			AppendCode(program.ByteCode, GetOpType());
-			AppendCode(program.ByteCode, Args);
+			++Program.NumOps;
+			LinkedAddress = MakeProgramAddress(GetOpType(), Program.ByteCode.Num());
+
+			AppendCode(Program.ByteCode, GetOpType());
+			AppendCode(Program.ByteCode, Args);
 		}
 	}
 
diff --git a/Engine/Plugins/Mutable/Source/MutableTools/Private/MuT/ASTOpImageInvert.cpp b/Engine/Plugins/Mutable/Source/MutableTools/Private/MuT/ASTOpImageInvert.cpp
--- a/Engine/Plugins/Mutable/Source/MutableTools/Private/MuT/ASTOpImageInvert.cpp
+++ b/Engine/Plugins/Mutable/Source/MutableTools/Private/MuT/ASTOpImageInvert.cpp
@@ -75,18 +75,19 @@ namespace UE::Mutable::Private
 	void ASTOpImageInvert::Link(FProgram& Program, FLinkerOptions*)
 	{
 		// Already linked?
-		if (!linkedAddress)
+		if (!LinkedAddress)
 		{
-			OP::ImageInvertArgs Args;
+			FOperation::ImageInvertArgs Args;
 			FMemory::Memzero(Args);
 
 			if (Base)
 			{
-				Args.Base = Base->linkedAddress;
+				Args.Base = Base->LinkedAddress;
 			}
 
-			linkedAddress = (OP::ADDRESS)Program.OpAddress.Num();
-			Program.OpAddress.Add(Program.ByteCode.Num());
+			++Program.NumOps;
+			LinkedAddress = MakeProgramAddress(GetOpType(), Program.ByteCode.Num());
+
 			AppendCode(Program.ByteCode, GetOpType());
 			AppendCode(Program.ByteCode, Args);
 		}
diff --git a/Engine/Plugins/Mutable/Source/MutableTools/Private/MuT/ASTOpImageLayer.cpp b/Engine/Plugins/Mutable/Source/MutableTools/Private/MuT/ASTOpImageLayer.cpp
--- a/Engine/Plugins/Mutable/Source/MutableTools/Private/MuT/ASTOpImageLayer.cpp
+++ b/Engine/Plugins/Mutable/Source/MutableTools/Private/MuT/ASTOpImageLayer.cpp
@@ -86,12 +86,12 @@ namespace UE::Mutable::Private
 	}
 
 
-	void ASTOpImageLayer::Link(FProgram& program, FLinkerOptions*)
+	void ASTOpImageLayer::Link(FProgram& Program, FLinkerOptions*)
 	{
 		// Already linked?
-		if (!linkedAddress)
+		if (!LinkedAddress)
 		{
-			OP::ImageLayerArgs Args;
+			FOperation::ImageLayerArgs Args;
 			FMemory::Memzero(Args);
 
 			Args.blendType = (uint8)blendType;
@@ -99,14 +99,26 @@ namespace UE::Mutable::Private
 			Args.BlendAlphaSourceChannel = BlendAlphaSourceChannel;
 			Args.flags = Flags;
 
-			if (base) Args.base = base->linkedAddress;
-			if (blend) Args.blended = blend->linkedAddress;
-			if (mask) Args.mask = mask->linkedAddress;
+			if (base) 
+			{
+				Args.base = base->LinkedAddress;
+			}
+
+			if (blend) 
+			{
+				Args.blended = blend->LinkedAddress;
+			}
 
-			linkedAddress = (OP::ADDRESS)program.OpAddress.Num();
-			program.OpAddress.Add((uint32)program.ByteCode.Num());
-			AppendCode(program.ByteCode, GetOpType());
-			AppendCode(program.ByteCode, Args);
+			if (mask)
+			{
+				Args.mask = mask->LinkedAddress;
+			}
+
+			++Program.NumOps;
+			LinkedAddress = MakeProgramAddress(GetOpType(), Program.ByteCode.Num());
+			
+			AppendCode(Program.ByteCode, GetOpType());
+			AppendCode(Program.ByteCode, Args);
 		}
 	}
 
@@ -371,7 +383,7 @@ namespace UE::Mutable::Private
 					// we can do something good here
 					Ptr<ASTOpImageLayer> NewLayer = UE::Mutable::Private::Clone<ASTOpImageLayer>(this);
 					NewLayer->mask = nullptr;
-					NewLayer->Flags |= OP::ImageLayerArgs::F_USE_MASK_FROM_BLENDED;
+					NewLayer->Flags |= FOperation::ImageLayerArgs::F_USE_MASK_FROM_BLENDED;
 					at = NewLayer;
 				}
 			}
@@ -432,7 +444,7 @@ namespace UE::Mutable::Private
 					if (ensure(SourceImageDesc.m_format == EImageFormat::RGBA_UByte))
 					{
 						Ptr<ASTOpImageLayer> nop = UE::Mutable::Private::Clone<ASTOpImageLayer>(this);
-						nop->Flags = Flags | OP::ImageLayerArgs::FLAGS::F_BLENDED_RGB_FROM_ALPHA;
+						nop->Flags = Flags | FOperation::ImageLayerArgs::FLAGS::F_BLENDED_RGB_FROM_ALPHA;
 						nop->blend = TypedBlend->Sources[0].child();
 						at = nop;
 					}
diff --git a/Engine/Plugins/Mutable/Source/MutableTools/Private/MuT/ASTOpImageLayerColor.cpp b/Engine/Plugins/Mutable/Source/MutableTools/Private/MuT/ASTOpImageLayerColor.cpp
--- a/Engine/Plugins/Mutable/Source/MutableTools/Private/MuT/ASTOpImageLayerColor.cpp
+++ b/Engine/Plugins/Mutable/Source/MutableTools/Private/MuT/ASTOpImageLayerColor.cpp
@@ -88,13 +88,12 @@ namespace UE::Mutable::Private
 	}
 
 
-	//-------------------------------------------------------------------------------------------------
-	void ASTOpImageLayerColor::Link(FProgram& program, FLinkerOptions*)
+	void ASTOpImageLayerColor::Link(FProgram& Program, FLinkerOptions*)
 	{
 		// Already linked?
-		if (!linkedAddress)
+		if (!LinkedAddress)
 		{
-			OP::ImageLayerColourArgs Args;
+			FOperation::ImageLayerColourArgs Args;
 			FMemory::Memzero(Args);
 
 			Args.blendType = (uint8)blendType;
@@ -103,14 +102,26 @@ namespace UE::Mutable::Private
 			Args.flags = Flags;
 
 			check(base);
-			if (base) Args.base = base->linkedAddress;
-			if (color) Args.colour = color->linkedAddress;
-			if (mask) Args.mask = mask->linkedAddress;
-
-			linkedAddress = (OP::ADDRESS)program.OpAddress.Num();
-			program.OpAddress.Add((uint32)program.ByteCode.Num());
-			AppendCode(program.ByteCode, GetOpType());
-			AppendCode(program.ByteCode, Args);
+			if (base)
+			{
+				Args.base = base->LinkedAddress;
+			}
+
+			if (color) 
+			{
+				Args.colour = color->LinkedAddress;
+			}
+
+			if (mask) 
+			{
+				Args.mask = mask->LinkedAddress;
+			}
+
+			++Program.NumOps;
+			LinkedAddress = MakeProgramAddress(GetOpType(), Program.ByteCode.Num());
+			
+			AppendCode(Program.ByteCode, GetOpType());
+			AppendCode(Program.ByteCode, Args);
 		}
 	}
 
@@ -181,7 +192,7 @@ namespace UE::Mutable::Private
 		Ptr<ASTOp> at;
 
 		// Plain masks optimization
-		if (mask.child() && !(Flags & OP::ImageLayerArgs::F_USE_MASK_FROM_BLENDED) )
+		if (mask.child() && !(Flags & FOperation::ImageLayerArgs::F_USE_MASK_FROM_BLENDED) )
 		{
 			FVector4f colour;
 			if (mask.child()->IsImagePlainConstant(colour))
@@ -230,7 +241,7 @@ namespace UE::Mutable::Private
 
 				if (!bRGBUnchanged)
 				{
-					if (Flags & OP::ImageLayerArgs::FLAGS::F_BASE_RGB_FROM_ALPHA)
+					if (Flags & FOperation::ImageLayerArgs::FLAGS::F_BASE_RGB_FROM_ALPHA)
 					{
 						switch (blendType)
 						{
@@ -309,7 +320,7 @@ namespace UE::Mutable::Private
 				if (bAreAllSameAlpha)
 				{
 					Ptr<ASTOpImageLayerColor> NewOp = UE::Mutable::Private::Clone<ASTOpImageLayerColor>(this);
-					NewOp->Flags |= OP::ImageLayerArgs::FLAGS::F_BASE_RGB_FROM_ALPHA;
+					NewOp->Flags |= FOperation::ImageLayerArgs::FLAGS::F_BASE_RGB_FROM_ALPHA;
 					NewOp->base = Source;
 					at = NewOp;
 				}
@@ -340,17 +351,17 @@ namespace UE::Mutable::Private
 //            {
 //                bModified = true;
 
-//                OP op = program.m_code[baseAt];
+//                OP op = Program.m_code[baseAt];
 
-//                OP aOp = program.m_code[at];
-//                aOp.Args.ImageLayerColour.base = program.m_code[baseAt].Args.Conditional.yes;
-//                op.args.Conditional.yes = program.AddOp( aOp );
+//                OP aOp = Program.m_code[at];
+//                aOp.Args.ImageLayerColour.base = Program.m_code[baseAt].Args.Conditional.yes;
+//                op.args.Conditional.yes = Program.AddOp( aOp );
 
-//                OP bOp = program.m_code[at];
-//                bOp.Args.ImageLayerColour.base = program.m_code[baseAt].Args.Conditional.no;
-//                op.args.Conditional.no = program.AddOp( bOp );
+//                OP bOp = Program.m_code[at];
+//                bOp.Args.ImageLayerColour.base = Program.m_code[baseAt].Args.Conditional.no;
+//                op.args.Conditional.no = Program.AddOp( bOp );
 
-//                at = program.AddOp( op );
+//                at = Program.AddOp( op );
 //                break;
 //            }
 
diff --git a/Engine/Plugins/Mutable/Source/MutableTools/Private/MuT/ASTOpImageLuminance.cpp b/Engine/Plugins/Mutable/Source/MutableTools/Private/MuT/ASTOpImageLuminance.cpp
--- a/Engine/Plugins/Mutable/Source/MutableTools/Private/MuT/ASTOpImageLuminance.cpp
+++ b/Engine/Plugins/Mutable/Source/MutableTools/Private/MuT/ASTOpImageLuminance.cpp
@@ -75,18 +75,19 @@ namespace UE::Mutable::Private
 	void ASTOpImageLuminance::Link(FProgram& Program, FLinkerOptions*)
 	{
 		// Already linked?
-		if (!linkedAddress)
+		if (!LinkedAddress)
 		{
-			OP::ImageLuminanceArgs Args;
+			FOperation::ImageLuminanceArgs Args;
 			FMemory::Memzero(Args);
 
 			if (Base)
 			{
-				Args.Base = Base->linkedAddress;
+				Args.Base = Base->LinkedAddress;
 			}
 
-			linkedAddress = (OP::ADDRESS)Program.OpAddress.Num();
-			Program.OpAddress.Add(Program.ByteCode.Num());
+			++Program.NumOps;
+			LinkedAddress = MakeProgramAddress(GetOpType(), Program.ByteCode.Num());
+			
 			AppendCode(Program.ByteCode, GetOpType());
 			AppendCode(Program.ByteCode, Args);
 		}
diff --git a/Engine/Plugins/Mutable/Source/MutableTools/Private/MuT/ASTOpImageMakeGrowMap.cpp b/Engine/Plugins/Mutable/Source/MutableTools/Private/MuT/ASTOpImageMakeGrowMap.cpp
--- a/Engine/Plugins/Mutable/Source/MutableTools/Private/MuT/ASTOpImageMakeGrowMap.cpp
+++ b/Engine/Plugins/Mutable/Source/MutableTools/Private/MuT/ASTOpImageMakeGrowMap.cpp
@@ -72,22 +72,25 @@ namespace UE::Mutable::Private
 	}
 
 
-	void ASTOpImageMakeGrowMap::Link(FProgram& program, FLinkerOptions*)
+	void ASTOpImageMakeGrowMap::Link(FProgram& Program, FLinkerOptions*)
 	{
 		// Already linked?
-		if (!linkedAddress)
+		if (!LinkedAddress)
 		{
-			OP::ImageMakeGrowMapArgs Args;
+			FOperation::ImageMakeGrowMapArgs Args;
 			FMemory::Memzero(Args);
 
 			Args.border = Border;
-			if (Mask) Args.mask = Mask->linkedAddress;
+			if (Mask) 
+			{
+				Args.mask = Mask->LinkedAddress;
+			}
 
-			linkedAddress = (OP::ADDRESS)program.OpAddress.Num();
-			//program.m_code.push_back(op);
-			program.OpAddress.Add((uint32_t)program.ByteCode.Num());
-			AppendCode(program.ByteCode, GetOpType());
-			AppendCode(program.ByteCode, Args);
+			++Program.NumOps;
+			LinkedAddress = MakeProgramAddress(GetOpType(), Program.ByteCode.Num());
+			
+			AppendCode(Program.ByteCode, GetOpType());
+			AppendCode(Program.ByteCode, Args);
 		}
 	}
 
diff --git a/Engine/Plugins/Mutable/Source/MutableTools/Private/MuT/ASTOpImageMipmap.cpp b/Engine/Plugins/Mutable/Source/MutableTools/Private/MuT/ASTOpImageMipmap.cpp
--- a/Engine/Plugins/Mutable/Source/MutableTools/Private/MuT/ASTOpImageMipmap.cpp
+++ b/Engine/Plugins/Mutable/Source/MutableTools/Private/MuT/ASTOpImageMipmap.cpp
@@ -86,12 +86,12 @@ namespace UE::Mutable::Private
 	}
 
 
-	void ASTOpImageMipmap::Link(FProgram& program, FLinkerOptions*)
+	void ASTOpImageMipmap::Link(FProgram& Program, FLinkerOptions*)
 	{
 		// Already linked?
-		if (!linkedAddress)
+		if (!LinkedAddress)
 		{
-			OP::ImageMipmapArgs Args;
+			FOperation::ImageMipmapArgs Args;
 			FMemory::Memzero(Args);
 
 			Args.levels = Levels;
@@ -99,13 +99,17 @@ namespace UE::Mutable::Private
 			Args.onlyTail = bOnlyTail;
 			Args.AddressMode = AddressMode;
 			Args.FilterType = FilterType;
-			if (Source) Args.source = Source->linkedAddress;
 
-			linkedAddress = (OP::ADDRESS)program.OpAddress.Num();
-			//program.m_code.push_back(op);
-			program.OpAddress.Add((uint32_t)program.ByteCode.Num());
-			AppendCode(program.ByteCode, GetOpType());
-			AppendCode(program.ByteCode, Args);
+			if (Source)
+			{
+				Args.source = Source->LinkedAddress;
+			}
+
+			++Program.NumOps;
+			LinkedAddress = MakeProgramAddress(GetOpType(), Program.ByteCode.Num());
+			
+			AppendCode(Program.ByteCode, GetOpType());
+			AppendCode(Program.ByteCode, Args);
 		}
 	}
 
diff --git a/Engine/Plugins/Mutable/Source/MutableTools/Private/MuT/ASTOpImageMultiCompose.cpp b/Engine/Plugins/Mutable/Source/MutableTools/Private/MuT/ASTOpImageMultiCompose.cpp
--- a/Engine/Plugins/Mutable/Source/MutableTools/Private/MuT/ASTOpImageMultiCompose.cpp
+++ b/Engine/Plugins/Mutable/Source/MutableTools/Private/MuT/ASTOpImageMultiCompose.cpp
@@ -100,29 +100,29 @@ void ASTOpImageMultiCompose::ForEachChild(const TFunctionRef<void(ASTChild&)> Fu
 void ASTOpImageMultiCompose::Link( FProgram& Program, FLinkerOptions*)
 {
     // Already linked?
-    if (!linkedAddress)
+    if (!LinkedAddress)
     {
-        OP::ImageMultiComposeArgs Args;
+        FOperation::ImageMultiComposeArgs Args;
 		FMemory::Memzero(Args);
 
 		if (Layout) 
 		{
-			Args.Layout = Layout->linkedAddress;
+			Args.Layout = Layout->LinkedAddress;
 		}
 
 		if (Base) 
 		{
-			Args.Base = Base->linkedAddress;
+			Args.Base = Base->LinkedAddress;
 		}
 
 		if (SourceLayout) 
 		{
-			Args.SourceLayout = SourceLayout->linkedAddress;
+			Args.SourceLayout = SourceLayout->LinkedAddress;
 		}
 
 		if (SourceImage) 
 		{
-			Args.SourceImage = SourceImage->linkedAddress;
+			Args.SourceImage = SourceImage->LinkedAddress;
 		}
 
 		Args.LayoutBlockSizeInPixelsX = LayoutBlockSizeInPixelsX;
@@ -131,11 +131,12 @@ void ASTOpImageMultiCompose::Link( FProgram& Program, FLinkerOptions*)
 		Args.SourceSizeX = SourceSizeX;
 		Args.SourceSizeY = SourceSizeY;
 
-        linkedAddress = (OP::ADDRESS)Program.OpAddress.Num();
-		Program.OpAddress.Add((uint32)Program.ByteCode.Num());
-        AppendCode(Program.ByteCode,EOpType::IM_MULTICOMPOSE);
-        AppendCode(Program.ByteCode, Args);
-    }
+		++Program.NumOps;
+		LinkedAddress = MakeProgramAddress(GetOpType(), Program.ByteCode.Num());
+		
+		AppendCode(Program.ByteCode, GetOpType());
+		AppendCode(Program.ByteCode, Args);
+	}
 
 }
 
diff --git a/Engine/Plugins/Mutable/Source/MutableTools/Private/MuT/ASTOpImageMultiLayer.cpp b/Engine/Plugins/Mutable/Source/MutableTools/Private/MuT/ASTOpImageMultiLayer.cpp
--- a/Engine/Plugins/Mutable/Source/MutableTools/Private/MuT/ASTOpImageMultiLayer.cpp
+++ b/Engine/Plugins/Mutable/Source/MutableTools/Private/MuT/ASTOpImageMultiLayer.cpp
@@ -91,12 +91,12 @@ namespace UE::Mutable::Private
 
 
 	//-------------------------------------------------------------------------------------------------
-	void ASTOpImageMultiLayer::Link(FProgram& program, FLinkerOptions*)
+	void ASTOpImageMultiLayer::Link(FProgram& Program, FLinkerOptions*)
 	{
 		// Already linked?
-		if (!linkedAddress)
+		if (!LinkedAddress)
 		{
-			OP::ImageMultiLayerArgs Args;
+			FOperation::ImageMultiLayerArgs Args;
 			FMemory::Memzero(Args);
 
 			Args.blendType = (uint8)blendType;
@@ -104,18 +104,31 @@ namespace UE::Mutable::Private
 			Args.BlendAlphaSourceChannel = BlendAlphaSourceChannel;
 			Args.bUseMaskFromBlended = bUseMaskFromBlended;
 
-			if (base) Args.base = base->linkedAddress;
-			if (blend) Args.blended = blend->linkedAddress;
-			if (mask) Args.mask = mask->linkedAddress;
+			if (base) 
+			{
+				Args.base = base->LinkedAddress;
+			}
+
+			if (blend) 
+			{
+				Args.blended = blend->LinkedAddress;
+			}
+
+			if (mask) 
+			{
+				Args.mask = mask->LinkedAddress;
+			}
+
 			if (range.rangeSize)
 			{
-				LinkRange(program, range, Args.rangeSize, Args.rangeId);
+				LinkRange(Program, range, Args.rangeSize, Args.rangeId);
 			}
 
-			linkedAddress = (OP::ADDRESS)program.OpAddress.Num();
-			program.OpAddress.Add((uint32)program.ByteCode.Num());
-			AppendCode(program.ByteCode, GetOpType());
-			AppendCode(program.ByteCode, Args);
+			++Program.NumOps;
+			LinkedAddress = MakeProgramAddress(GetOpType(), Program.ByteCode.Num());
+			
+			AppendCode(Program.ByteCode, GetOpType());
+			AppendCode(Program.ByteCode, Args);
 		}
 	}
 
diff --git a/Engine/Plugins/Mutable/Source/MutableTools/Private/MuT/ASTOpImageNormalComposite.cpp b/Engine/Plugins/Mutable/Source/MutableTools/Private/MuT/ASTOpImageNormalComposite.cpp
--- a/Engine/Plugins/Mutable/Source/MutableTools/Private/MuT/ASTOpImageNormalComposite.cpp
+++ b/Engine/Plugins/Mutable/Source/MutableTools/Private/MuT/ASTOpImageNormalComposite.cpp
@@ -75,31 +75,32 @@ namespace UE::Mutable::Private
 	}
 
 
-	void ASTOpImageNormalComposite::Link(FProgram& program, FLinkerOptions*)
+	void ASTOpImageNormalComposite::Link(FProgram& Program, FLinkerOptions*)
 	{
 		// Already linked?
-		if (!linkedAddress)
+		if (!LinkedAddress)
 		{
-			OP::ImageNormalCompositeArgs Args;
+			FOperation::ImageNormalCompositeArgs Args;
 			FMemory::Memzero(Args);
 
 			if (Base)
 			{
-				Args.base = Base->linkedAddress;
+				Args.base = Base->LinkedAddress;
 			}
 
 			if (Normal)
 			{
-				Args.normal = Normal->linkedAddress;
+				Args.normal = Normal->LinkedAddress;
 			}
 
 			Args.power = Power;
 			Args.mode = Mode;
 
-			linkedAddress = (OP::ADDRESS)program.OpAddress.Num();
-			program.OpAddress.Add((uint32_t)program.ByteCode.Num());
-			AppendCode(program.ByteCode, EOpType::IM_NORMALCOMPOSITE);
-			AppendCode(program.ByteCode, Args);
+			++Program.NumOps;
+			LinkedAddress = MakeProgramAddress(GetOpType(), Program.ByteCode.Num());
+			
+			AppendCode(Program.ByteCode, GetOpType());
+			AppendCode(Program.ByteCode, Args);
 		}
 
 	}
diff --git a/Engine/Plugins/Mutable/Source/MutableTools/Private/MuT/ASTOpImageParameterConvert.cpp b/Engine/Plugins/Mutable/Source/MutableTools/Private/MuT/ASTOpImageParameterConvert.cpp
--- a/Engine/Plugins/Mutable/Source/MutableTools/Private/MuT/ASTOpImageParameterConvert.cpp
+++ b/Engine/Plugins/Mutable/Source/MutableTools/Private/MuT/ASTOpImageParameterConvert.cpp
@@ -69,18 +69,18 @@ namespace UE::Mutable::Private
 
 	void ASTOpImageParameterConvert::Link(FProgram& Program, FLinkerOptions*)
 	{
-		if (!linkedAddress)
+		if (!LinkedAddress)
 		{
-			linkedAddress = static_cast<OP::ADDRESS>(Program.OpAddress.Num());
-			Program.OpAddress.Add(static_cast<uint32>(Program.ByteCode.Num()));
-			AppendCode(Program.ByteCode, Type);
-
-			OP::ImageParameterConvertArgs Args;
+			FOperation::ImageParameterConvertArgs Args;
 			FMemory::Memzero(Args);
 
-			check(ImageParameter->linkedAddress);
-			Args.ImageParameter = ImageParameter->linkedAddress;
+			check(ImageParameter->LinkedAddress);
+			Args.ImageParameter = ImageParameter->LinkedAddress;
 		
+			++Program.NumOps;
+			LinkedAddress = MakeProgramAddress(GetOpType(), Program.ByteCode.Num());
+			
+			AppendCode(Program.ByteCode, GetOpType());
 			AppendCode(Program.ByteCode, Args);
 		}
 	}
diff --git a/Engine/Plugins/Mutable/Source/MutableTools/Private/MuT/ASTOpImagePatch.cpp b/Engine/Plugins/Mutable/Source/MutableTools/Private/MuT/ASTOpImagePatch.cpp
--- a/Engine/Plugins/Mutable/Source/MutableTools/Private/MuT/ASTOpImagePatch.cpp
+++ b/Engine/Plugins/Mutable/Source/MutableTools/Private/MuT/ASTOpImagePatch.cpp
@@ -71,26 +71,33 @@ namespace UE::Mutable::Private
 	}
 
 
-	//-------------------------------------------------------------------------------------------------
-	void ASTOpImagePatch::Link(FProgram& program, FLinkerOptions*)
+	void ASTOpImagePatch::Link(FProgram& Program, FLinkerOptions*)
 	{
 		// Already linked?
-		if (!linkedAddress)
+		if (!LinkedAddress)
 		{
-			OP::ImagePatchArgs Args;
+			FOperation::ImagePatchArgs Args;
 			FMemory::Memzero(Args);
 
-			if (base) Args.base = base->linkedAddress;
-			if (patch) Args.patch = patch->linkedAddress;
+			if (base)
+			{
+				Args.base = base->LinkedAddress;
+			}
+
+			if (patch) 
+			{
+				Args.patch = patch->LinkedAddress;
+			}
+
 			Args.minX = location[0];
 			Args.minY = location[1];
 
-			linkedAddress = (OP::ADDRESS)program.OpAddress.Num();
-			program.OpAddress.Add((uint32_t)program.ByteCode.Num());
-			AppendCode(program.ByteCode, EOpType::IM_PATCH);
-			AppendCode(program.ByteCode, Args);
+			++Program.NumOps;
+			LinkedAddress = MakeProgramAddress(GetOpType(), Program.ByteCode.Num());
+			
+			AppendCode(Program.ByteCode, GetOpType());
+			AppendCode(Program.ByteCode, Args);
 		}
-
 	}
 
 
diff --git a/Engine/Plugins/Mutable/Source/MutableTools/Private/MuT/ASTOpImagePixelFormat.cpp b/Engine/Plugins/Mutable/Source/MutableTools/Private/MuT/ASTOpImagePixelFormat.cpp
--- a/Engine/Plugins/Mutable/Source/MutableTools/Private/MuT/ASTOpImagePixelFormat.cpp
+++ b/Engine/Plugins/Mutable/Source/MutableTools/Private/MuT/ASTOpImagePixelFormat.cpp
@@ -85,23 +85,26 @@ namespace UE::Mutable::Private
 	}
 
 
-	void ASTOpImagePixelFormat::Link(FProgram& program, FLinkerOptions*)
+	void ASTOpImagePixelFormat::Link(FProgram& Program, FLinkerOptions*)
 	{
 		// Already linked?
-		if (!linkedAddress)
+		if (!LinkedAddress)
 		{
-			OP::ImagePixelFormatArgs Args;
+			FOperation::ImagePixelFormatArgs Args;
 			FMemory::Memzero(Args);
 
 			Args.format = Format;
 			Args.formatIfAlpha = FormatIfAlpha;
-			if (Source) Args.source = Source->linkedAddress;
+			if (Source) 
+			{
+				Args.source = Source->LinkedAddress;
+			}
+
+			++Program.NumOps;
+			LinkedAddress = MakeProgramAddress(GetOpType(), Program.ByteCode.Num());
 
-			linkedAddress = (OP::ADDRESS)program.OpAddress.Num();
-			//program.m_code.push_back(op);
-			program.OpAddress.Add((uint32_t)program.ByteCode.Num());
-			AppendCode(program.ByteCode, EOpType::IM_PIXELFORMAT);
-			AppendCode(program.ByteCode, Args);
+			AppendCode(Program.ByteCode, GetOpType());
+			AppendCode(Program.ByteCode, Args);
 		}
 
 	}
diff --git a/Engine/Plugins/Mutable/Source/MutableTools/Private/MuT/ASTOpImagePlainColor.cpp b/Engine/Plugins/Mutable/Source/MutableTools/Private/MuT/ASTOpImagePlainColor.cpp
--- a/Engine/Plugins/Mutable/Source/MutableTools/Private/MuT/ASTOpImagePlainColor.cpp
+++ b/Engine/Plugins/Mutable/Source/MutableTools/Private/MuT/ASTOpImagePlainColor.cpp
@@ -79,14 +79,14 @@ namespace UE::Mutable::Private
 	void ASTOpImagePlainColor::Link(FProgram& Program, FLinkerOptions*)
 	{
 		// Already linked?
-		if (!linkedAddress)
+		if (!LinkedAddress)
 		{
-			OP::ImagePlainColorArgs Args;
+			FOperation::ImagePlainColorArgs Args;
 			FMemory::Memzero(Args);
 
 			if (Color)
 			{
-				Args.Color = Color->linkedAddress;
+				Args.Color = Color->LinkedAddress;
 			}
 
 			Args.Format = Format;
@@ -94,8 +94,9 @@ namespace UE::Mutable::Private
 			Args.Size[1] = Size[1];
 			Args.LODs = LODs;
 
-			linkedAddress = (OP::ADDRESS)Program.OpAddress.Num();
-			Program.OpAddress.Add(Program.ByteCode.Num());
+			++Program.NumOps;
+			LinkedAddress = MakeProgramAddress(GetOpType(), Program.ByteCode.Num());
+			
 			AppendCode(Program.ByteCode, GetOpType());
 			AppendCode(Program.ByteCode, Args);
 		}
diff --git a/Engine/Plugins/Mutable/Source/MutableTools/Private/MuT/ASTOpImageRasterMesh.cpp b/Engine/Plugins/Mutable/Source/MutableTools/Private/MuT/ASTOpImageRasterMesh.cpp
--- a/Engine/Plugins/Mutable/Source/MutableTools/Private/MuT/ASTOpImageRasterMesh.cpp
+++ b/Engine/Plugins/Mutable/Source/MutableTools/Private/MuT/ASTOpImageRasterMesh.cpp
@@ -134,12 +134,12 @@ namespace UE::Mutable::Private
 
 
 	//-------------------------------------------------------------------------------------------------
-	void ASTOpImageRasterMesh::Link(FProgram& program, FLinkerOptions*)
+	void ASTOpImageRasterMesh::Link(FProgram& Program, FLinkerOptions*)
 	{
 		// Already linked?
-		if (!linkedAddress)
+		if (!LinkedAddress)
 		{
-			OP::ImageRasterMeshArgs Args;
+			FOperation::ImageRasterMeshArgs Args;
 			FMemory::Memzero(Args);
 
 			Args.BlockId = BlockId;
@@ -157,16 +157,36 @@ namespace UE::Mutable::Private
 			Args.SamplingMethod = static_cast<uint8>(SamplingMethod);
 			Args.MinFilterMethod = static_cast<uint8>(MinFilterMethod);
 
-			if (mesh) Args.mesh = mesh->linkedAddress;
-			if (image) Args.image = image->linkedAddress;
-			if (angleFadeProperties) Args.angleFadeProperties = angleFadeProperties->linkedAddress;
-			if (mask) Args.mask = mask->linkedAddress;
-			if (projector) Args.projector = projector->linkedAddress;
+			if (mesh) 
+			{
+				Args.mesh = mesh->LinkedAddress;
+			}
+
+			if (image) 
+			{
+				Args.image = image->LinkedAddress;
+			}
+
+			if (angleFadeProperties) 
+			{
+				Args.angleFadeProperties = angleFadeProperties->LinkedAddress;
+			}
+
+			if (mask) 
+			{
+				Args.mask = mask->LinkedAddress;
+			}
+
+			if (projector) 
+			{
+				Args.projector = projector->LinkedAddress;
+			}
 
-			linkedAddress = (OP::ADDRESS)program.OpAddress.Num();
-			program.OpAddress.Add((uint32)program.ByteCode.Num());
-			AppendCode(program.ByteCode, GetOpType());
-			AppendCode(program.ByteCode, Args);
+			++Program.NumOps;
+			LinkedAddress = MakeProgramAddress(GetOpType(), Program.ByteCode.Num());
+			
+			AppendCode(Program.ByteCode, GetOpType());
+			AppendCode(Program.ByteCode, Args);
 		}
 	}
 
diff --git a/Engine/Plugins/Mutable/Source/MutableTools/Private/MuT/ASTOpImageResize.cpp b/Engine/Plugins/Mutable/Source/MutableTools/Private/MuT/ASTOpImageResize.cpp
--- a/Engine/Plugins/Mutable/Source/MutableTools/Private/MuT/ASTOpImageResize.cpp
+++ b/Engine/Plugins/Mutable/Source/MutableTools/Private/MuT/ASTOpImageResize.cpp
@@ -82,20 +82,22 @@ namespace UE::Mutable::Private
 	void ASTOpImageResize::Link(FProgram& Program, FLinkerOptions*)
 	{
 		// Already linked?
-		if (!linkedAddress)
+		if (!LinkedAddress)
 		{
-			OP::ImageResizeArgs Args;
+			FOperation::ImageResizeArgs Args;
 			FMemory::Memzero(Args);
 
 			if (Source)
 			{
-				Args.Source = Source->linkedAddress;
+				Args.Source = Source->LinkedAddress;
 			}
+
 			Args.Size[0] = Size[0];
 			Args.Size[1] = Size[1];
 
-			linkedAddress = (OP::ADDRESS)Program.OpAddress.Num();
-			Program.OpAddress.Add(Program.ByteCode.Num());
+			++Program.NumOps;
+			LinkedAddress = MakeProgramAddress(GetOpType(), Program.ByteCode.Num());
+			
 			AppendCode(Program.ByteCode, GetOpType());
 			AppendCode(Program.ByteCode, Args);
 		}
diff --git a/Engine/Plugins/Mutable/Source/MutableTools/Private/MuT/ASTOpImageResizeLike.cpp b/Engine/Plugins/Mutable/Source/MutableTools/Private/MuT/ASTOpImageResizeLike.cpp
--- a/Engine/Plugins/Mutable/Source/MutableTools/Private/MuT/ASTOpImageResizeLike.cpp
+++ b/Engine/Plugins/Mutable/Source/MutableTools/Private/MuT/ASTOpImageResizeLike.cpp
@@ -80,22 +80,24 @@ namespace UE::Mutable::Private
 	void ASTOpImageResizeLike::Link(FProgram& Program, FLinkerOptions*)
 	{
 		// Already linked?
-		if (!linkedAddress)
+		if (!LinkedAddress)
 		{
-			OP::ImageResizeLikeArgs Args;
+			FOperation::ImageResizeLikeArgs Args;
 			FMemory::Memzero(Args);
 
 			if (Source)
 			{
-				Args.Source = Source->linkedAddress;
+				Args.Source = Source->LinkedAddress;
 			}
+
 			if (SizeSource)
 			{
-				Args.SizeSource = SizeSource->linkedAddress;
+				Args.SizeSource = SizeSource->LinkedAddress;
 			}
 
-			linkedAddress = (OP::ADDRESS)Program.OpAddress.Num();
-			Program.OpAddress.Add(Program.ByteCode.Num());
+			++Program.NumOps;
+			LinkedAddress = MakeProgramAddress(GetOpType(), Program.ByteCode.Num());
+			
 			AppendCode(Program.ByteCode, GetOpType());
 			AppendCode(Program.ByteCode, Args);
 		}
@@ -209,7 +211,7 @@ namespace UE::Mutable::Private
 		{
 			// TODO
 			// Skip intermediate ops until the layout
-//                    if ( program.m_code[sizeSourceAt].type != EOpType::IM_BLANKLAYOUT )
+//                    if ( Program.m_code[sizeSourceAt].type != EOpType::IM_BLANKLAYOUT )
 //                    {
 //                        m_modified = true;
 
@@ -220,9 +222,9 @@ namespace UE::Mutable::Private
 //                        blackLayoutOp.args.ImageBlankLayout.blockSize[1] = sizeSourceSize->factor[1];
 //                        blackLayoutOp.args.ImageBlankLayout.format = L_UByte;
 
-//                        OP opResize = program.m_code[at];
-//                        opResize.args.ImageResizeLike.sizeSource = program.AddOp( blackLayoutOp );
-//                        Result = program.AddOp( opResize );
+//                        OP opResize = Program.m_code[at];
+//                        opResize.args.ImageResizeLike.sizeSource = Program.AddOp( blackLayoutOp );
+//                        Result = Program.AddOp( opResize );
 //                    }
 		}
 
diff --git a/Engine/Plugins/Mutable/Source/MutableTools/Private/MuT/ASTOpImageResizeRel.cpp b/Engine/Plugins/Mutable/Source/MutableTools/Private/MuT/ASTOpImageResizeRel.cpp
--- a/Engine/Plugins/Mutable/Source/MutableTools/Private/MuT/ASTOpImageResizeRel.cpp
+++ b/Engine/Plugins/Mutable/Source/MutableTools/Private/MuT/ASTOpImageResizeRel.cpp
@@ -87,20 +87,22 @@ namespace UE::Mutable::Private
 	void ASTOpImageResizeRel::Link(FProgram& Program, FLinkerOptions*)
 	{
 		// Already linked?
-		if (!linkedAddress)
+		if (!LinkedAddress)
 		{
-			OP::ImageResizeRelArgs Args;
+			FOperation::ImageResizeRelArgs Args;
 			FMemory::Memzero(Args);
 
 			if (Source)
 			{
-				Args.Source = Source->linkedAddress;
+				Args.Source = Source->LinkedAddress;
 			}
+
 			Args.Factor[0] = Factor[0];
 			Args.Factor[1] = Factor[1];
 
-			linkedAddress = (OP::ADDRESS)Program.OpAddress.Num();
-			Program.OpAddress.Add(Program.ByteCode.Num());
+			++Program.NumOps;
+			LinkedAddress = MakeProgramAddress(GetOpType(), Program.ByteCode.Num());
+			
 			AppendCode(Program.ByteCode, GetOpType());
 			AppendCode(Program.ByteCode, Args);
 		}
diff --git a/Engine/Plugins/Mutable/Source/MutableTools/Private/MuT/ASTOpImageSaturate.cpp b/Engine/Plugins/Mutable/Source/MutableTools/Private/MuT/ASTOpImageSaturate.cpp
--- a/Engine/Plugins/Mutable/Source/MutableTools/Private/MuT/ASTOpImageSaturate.cpp
+++ b/Engine/Plugins/Mutable/Source/MutableTools/Private/MuT/ASTOpImageSaturate.cpp
@@ -80,22 +80,24 @@ namespace UE::Mutable::Private
 	void ASTOpImageSaturate::Link(FProgram& Program, FLinkerOptions*)
 	{
 		// Already linked?
-		if (!linkedAddress)
+		if (!LinkedAddress)
 		{
-			OP::ImageSaturateArgs Args;
+			FOperation::ImageSaturateArgs Args;
 			FMemory::Memzero(Args);
 
 			if (Base)
 			{
-				Args.Base = Base->linkedAddress;
+				Args.Base = Base->LinkedAddress;
 			}
+
 			if (Factor)
 			{
-				Args.Factor = Factor->linkedAddress;
+				Args.Factor = Factor->LinkedAddress;
 			}
 
-			linkedAddress = (OP::ADDRESS)Program.OpAddress.Num();
-			Program.OpAddress.Add(Program.ByteCode.Num());
+			++Program.NumOps;
+			LinkedAddress = MakeProgramAddress(GetOpType(), Program.ByteCode.Num());
+			
 			AppendCode(Program.ByteCode, GetOpType());
 			AppendCode(Program.ByteCode, Args);
 		}
diff --git a/Engine/Plugins/Mutable/Source/MutableTools/Private/MuT/ASTOpImageSwizzle.cpp b/Engine/Plugins/Mutable/Source/MutableTools/Private/MuT/ASTOpImageSwizzle.cpp
--- a/Engine/Plugins/Mutable/Source/MutableTools/Private/MuT/ASTOpImageSwizzle.cpp
+++ b/Engine/Plugins/Mutable/Source/MutableTools/Private/MuT/ASTOpImageSwizzle.cpp
@@ -104,25 +104,30 @@ namespace UE::Mutable::Private
 	}
 
 
-	void ASTOpImageSwizzle::Link(FProgram& program, FLinkerOptions*)
+	void ASTOpImageSwizzle::Link(FProgram& Program, FLinkerOptions*)
 	{
 		// Already linked?
-		if (!linkedAddress)
+		if (!LinkedAddress)
 		{
-			OP::ImageSwizzleArgs Args;
+			FOperation::ImageSwizzleArgs Args;
 			FMemory::Memzero(Args);
 
 			Args.format = Format;
 			for (int32 i = 0; i < MUTABLE_OP_MAX_SWIZZLE_CHANNELS; ++i)
 			{
-				if (Sources[i]) Args.sources[i] = Sources[i]->linkedAddress;
+				if (Sources[i])
+				{
+					Args.sources[i] = Sources[i]->LinkedAddress;
+				}
+
 				Args.sourceChannels[i] = SourceChannels[i];
 			}			
 
-			linkedAddress = (OP::ADDRESS)program.OpAddress.Num();
-			program.OpAddress.Add((uint32)program.ByteCode.Num());
-			AppendCode(program.ByteCode, GetOpType());
-			AppendCode(program.ByteCode, Args);
+			++Program.NumOps;
+			LinkedAddress = MakeProgramAddress(GetOpType(), Program.ByteCode.Num());
+			
+			AppendCode(Program.ByteCode, GetOpType());
+			AppendCode(Program.ByteCode, Args);
 		}
 	}
 
@@ -386,7 +391,7 @@ namespace UE::Mutable::Private
 						)
 					{
 						// Additional optimization is possible here.
-						NewLayer->Flags |= OP::ImageLayerArgs::FLAGS::F_USE_MASK_FROM_BLENDED;
+						NewLayer->Flags |= FOperation::ImageLayerArgs::FLAGS::F_USE_MASK_FROM_BLENDED;
 						NewLayer->mask = nullptr;
 					}
 
@@ -642,10 +647,10 @@ namespace UE::Mutable::Private
 			//			const ASTOpFixed* Typed = static_cast<const ASTOpFixed*>(Sources[c].child().get());
 			//			check(Typed);
 			//			// Compare all Args but the source image
-			//			OP::ImageResizeArgs ArgCopy = FirstResize->op.args.ImageResize;
+			//			FOperation::ImageResizeArgs ArgCopy = FirstResize->op.args.ImageResize;
 			//			ArgCopy.source = Typed->op.args.ImageResize.source;
 
-			//			if (FMemory::Memcmp(&ArgCopy, &Typed->op.args.ImageResize, sizeof(OP::ImageResizeArgs)) != 0)
+			//			if (FMemory::Memcmp(&ArgCopy, &Typed->op.args.ImageResize, sizeof(FOperation::ImageResizeArgs)) != 0)
 			//			{
 			//				bAreAllResizesCompatible = false;
 			//				break;
@@ -927,7 +932,7 @@ namespace UE::Mutable::Private
 			// For now just check the case that we are observing in the working data: 
 			// A is in the blended of a multiply, and we take its alpha channel
 			// \TODO: Implement the other cases when we find instances of them.
-			if ( OldLayer->Flags==OP::ImageLayerArgs::FLAGS::F_BLENDED_RGB_FROM_ALPHA
+			if ( OldLayer->Flags==FOperation::ImageLayerArgs::FLAGS::F_BLENDED_RGB_FROM_ALPHA
 				&&
 				bOldLayerBlendIsCompatibleWithSwizzleRGBs
 				&& 
diff --git a/Engine/Plugins/Mutable/Source/MutableTools/Private/MuT/ASTOpImageTransform.cpp b/Engine/Plugins/Mutable/Source/MutableTools/Private/MuT/ASTOpImageTransform.cpp
--- a/Engine/Plugins/Mutable/Source/MutableTools/Private/MuT/ASTOpImageTransform.cpp
+++ b/Engine/Plugins/Mutable/Source/MutableTools/Private/MuT/ASTOpImageTransform.cpp
@@ -114,17 +114,17 @@ namespace UE::Mutable::Private
 	void ASTOpImageTransform::Link(FProgram& Program, FLinkerOptions*)
 	{
 		// Already linked?
-		if (!linkedAddress)
+		if (!LinkedAddress)
 		{
-			OP::ImageTransformArgs Args;
+			FOperation::ImageTransformArgs Args;
 			FMemory::Memzero(Args);
 
-			Args.Base     = Base     ? Base->linkedAddress     : 0;
-			Args.OffsetX  = OffsetX  ? OffsetX->linkedAddress  : 0;
-			Args.OffsetY  = OffsetY  ? OffsetY->linkedAddress  : 0;
-			Args.ScaleX   = ScaleX   ? ScaleX->linkedAddress   : 0;
-			Args.ScaleY   = ScaleY   ? ScaleY->linkedAddress   : 0;
-			Args.Rotation = Rotation ? Rotation->linkedAddress : 0;
+			Args.Base     = Base     ? Base->LinkedAddress     : 0;
+			Args.OffsetX  = OffsetX  ? OffsetX->LinkedAddress  : 0;
+			Args.OffsetY  = OffsetY  ? OffsetY->LinkedAddress  : 0;
+			Args.ScaleX   = ScaleX   ? ScaleX->LinkedAddress   : 0;
+			Args.ScaleY   = ScaleY   ? ScaleY->LinkedAddress   : 0;
+			Args.Rotation = Rotation ? Rotation->LinkedAddress : 0;
 			Args.AddressMode = static_cast<uint32>(AddressMode);
 			Args.bKeepAspectRatio = bKeepAspectRatio;
 			Args.SizeX = SizeX;
@@ -132,9 +132,10 @@ namespace UE::Mutable::Private
 			Args.SourceSizeX = SourceSizeX;
 			Args.SourceSizeY = SourceSizeY;
 
-			linkedAddress = (OP::ADDRESS)Program.OpAddress.Num();
-			Program.OpAddress.Add((uint32_t)Program.ByteCode.Num());
-			AppendCode(Program.ByteCode, EOpType::IM_TRANSFORM);
+			++Program.NumOps;
+			LinkedAddress = MakeProgramAddress(GetOpType(), Program.ByteCode.Num());
+			
+			AppendCode(Program.ByteCode, GetOpType());
 			AppendCode(Program.ByteCode, Args);
 		}
 	}
diff --git a/Engine/Plugins/Mutable/Source/MutableTools/Private/MuT/ASTOpInstanceAdd.cpp b/Engine/Plugins/Mutable/Source/MutableTools/Private/MuT/ASTOpInstanceAdd.cpp
--- a/Engine/Plugins/Mutable/Source/MutableTools/Private/MuT/ASTOpInstanceAdd.cpp
+++ b/Engine/Plugins/Mutable/Source/MutableTools/Private/MuT/ASTOpInstanceAdd.cpp
@@ -98,9 +98,9 @@ namespace UE::Mutable::Private
 		check(type != EOpType::NONE);
 		
 		// Already linked?
-		if (!linkedAddress)
+		if (!LinkedAddress)
 		{
-			OP::InstanceAddArgs Args;
+			FOperation::InstanceAddArgs Args;
 			FMemory::Memzero(Args);
 			Args.id = id;
 			Args.ExternalId = ExternalId;
@@ -108,12 +108,12 @@ namespace UE::Mutable::Private
 
 			if (instance)
 			{
-				Args.instance = instance->linkedAddress;
+				Args.instance = instance->LinkedAddress;
 			}
 				
 			if (value)
 			{
-				Args.value = value->linkedAddress;
+				Args.value = value->LinkedAddress;
 
 				if (type == EOpType::IN_ADDMESH)
 				{
@@ -142,18 +142,19 @@ namespace UE::Mutable::Private
 
 					if (It != INDEX_NONE)
 					{
-						Program.RelevantParameterList.Add(value->linkedAddress, It);
+						Program.RelevantParameterList.Add(value->LinkedAddress, It);
 					}
 					else
 					{
-						Program.RelevantParameterList.Add(value->linkedAddress, Program.ParameterLists.Num());
+						Program.RelevantParameterList.Add(value->LinkedAddress, Program.ParameterLists.Num());
 						Program.ParameterLists.Add(MatchingParameterIndices);
 					}
 				}
 			}
 
-			linkedAddress = (OP::ADDRESS)Program.OpAddress.Num();
-			Program.OpAddress.Add((uint32_t)Program.ByteCode.Num());
+			++Program.NumOps;
+			LinkedAddress = MakeProgramAddress(type, Program.ByteCode.Num());
+
 			AppendCode(Program.ByteCode, type);
 			AppendCode(Program.ByteCode, Args);
 		}
@@ -180,7 +181,7 @@ namespace UE::Mutable::Private
 			{
 
 			// We want to move the conditional up the op graph because this way the mesh root operation
-			// can easily match other mesh root operations in the program. This is important because this operation
+			// can easily match other mesh root operations in the Program. This is important because this operation
 			// address is used for caching, and avoiding duplicate mesh work with multi-components.
 			case EOpType::ME_CONDITIONAL:
 			{
diff --git a/Engine/Plugins/Mutable/Source/MutableTools/Private/MuT/ASTOpLayoutFromMesh.cpp b/Engine/Plugins/Mutable/Source/MutableTools/Private/MuT/ASTOpLayoutFromMesh.cpp
--- a/Engine/Plugins/Mutable/Source/MutableTools/Private/MuT/ASTOpLayoutFromMesh.cpp
+++ b/Engine/Plugins/Mutable/Source/MutableTools/Private/MuT/ASTOpLayoutFromMesh.cpp
@@ -43,12 +43,12 @@ namespace UE::Mutable::Private
 	}
 
 
-	bool ASTOpLayoutFromMesh::IsEqual(const ASTOp& otherUntyped) const
+	bool ASTOpLayoutFromMesh::IsEqual(const ASTOp& OtherUntyped) const
 	{
-		if (otherUntyped.GetOpType()==GetOpType())
+		if (OtherUntyped.GetOpType() == GetOpType())
 		{
-			const ASTOpLayoutFromMesh* other = static_cast<const ASTOpLayoutFromMesh*>(&otherUntyped);
-			return Mesh == other->Mesh && LayoutIndex == other->LayoutIndex;
+			const ASTOpLayoutFromMesh* Other = static_cast<const ASTOpLayoutFromMesh*>(&OtherUntyped);
+			return Mesh == Other->Mesh && LayoutIndex == Other->LayoutIndex;
 		}
 		return false;
 	}
@@ -65,37 +65,40 @@ namespace UE::Mutable::Private
 	}
 
 
-	UE::Mutable::Private::Ptr<ASTOp> ASTOpLayoutFromMesh::Clone(MapChildFuncRef mapChild) const
+	UE::Mutable::Private::Ptr<ASTOp> ASTOpLayoutFromMesh::Clone(MapChildFuncRef MapChild) const
 	{
-		UE::Mutable::Private::Ptr<ASTOpLayoutFromMesh> n = new ASTOpLayoutFromMesh();
-		n->Mesh = mapChild(Mesh.child());
-		n->LayoutIndex = LayoutIndex;
-		return n;
+		UE::Mutable::Private::Ptr<ASTOpLayoutFromMesh> New = new ASTOpLayoutFromMesh();
+		New->Mesh = MapChild(Mesh.child());
+		New->LayoutIndex = LayoutIndex;
+		return New;
 	}
 
 
-	void ASTOpLayoutFromMesh::ForEachChild(const TFunctionRef<void(ASTChild&)> f)
+	void ASTOpLayoutFromMesh::ForEachChild(const TFunctionRef<void(ASTChild&)> Func)
 	{
-		f(Mesh);
+		Func(Mesh);
 	}
 
 
-	void ASTOpLayoutFromMesh::Link(FProgram& program, FLinkerOptions*)
+	void ASTOpLayoutFromMesh::Link(FProgram& Program, FLinkerOptions*)
 	{
 		// Already linked?
-		if (!linkedAddress)
+		if (!LinkedAddress)
 		{
-			OP::LayoutFromMeshArgs Args;
+			FOperation::LayoutFromMeshArgs Args;
 			FMemory::Memzero(Args);
 
 			Args.LayoutIndex = LayoutIndex;
-			if (Mesh) Args.Mesh = Mesh->linkedAddress;
+			if (Mesh) 
+			{
+				Args.Mesh = Mesh->LinkedAddress;
+			}
 
-			linkedAddress = (OP::ADDRESS)program.OpAddress.Num();
-			//program.m_code.push_back(op);
-			program.OpAddress.Add((uint32_t)program.ByteCode.Num());
-			AppendCode(program.ByteCode, GetOpType());
-			AppendCode(program.ByteCode, Args);
+			++Program.NumOps;
+			LinkedAddress = MakeProgramAddress(GetOpType(), Program.ByteCode.Num());
+			
+			AppendCode(Program.ByteCode, GetOpType());
+			AppendCode(Program.ByteCode, Args);
 		}
 
 	}
diff --git a/Engine/Plugins/Mutable/Source/MutableTools/Private/MuT/ASTOpLayoutMerge.cpp b/Engine/Plugins/Mutable/Source/MutableTools/Private/MuT/ASTOpLayoutMerge.cpp
--- a/Engine/Plugins/Mutable/Source/MutableTools/Private/MuT/ASTOpLayoutMerge.cpp
+++ b/Engine/Plugins/Mutable/Source/MutableTools/Private/MuT/ASTOpLayoutMerge.cpp
@@ -23,7 +23,6 @@ namespace UE::Mutable::Private
 {
 
 
-	//-------------------------------------------------------------------------------------------------
 	ASTOpLayoutMerge::ASTOpLayoutMerge()
 		: Base(this)
 		, Added(this)
@@ -76,23 +75,30 @@ namespace UE::Mutable::Private
 	}
 
 
-	void ASTOpLayoutMerge::Link(FProgram& program, FLinkerOptions*)
+	void ASTOpLayoutMerge::Link(FProgram& Program, FLinkerOptions*)
 	{
 		// Already linked?
-		if (!linkedAddress)
+		if (!LinkedAddress)
 		{
-			OP::LayoutMergeArgs Args;
+			FOperation::LayoutMergeArgs Args;
 			FMemory::Memzero(Args);
 
-			if (Base) Args.Base = Base->linkedAddress;
-			if (Added) Args.Added = Added->linkedAddress;
+			if (Base) 
+			{
+				Args.Base = Base->LinkedAddress;
+			}
 
-			linkedAddress = (OP::ADDRESS)program.OpAddress.Num();
-			program.OpAddress.Add((uint32_t)program.ByteCode.Num());
-			AppendCode(program.ByteCode, GetOpType());
-			AppendCode(program.ByteCode, Args);
-		}
+			if (Added) 
+			{
+				Args.Added = Added->LinkedAddress;
+			}
+
+			++Program.NumOps;
+			LinkedAddress = MakeProgramAddress(GetOpType(), Program.ByteCode.Num());
 
+			AppendCode(Program.ByteCode, GetOpType());
+			AppendCode(Program.ByteCode, Args);
+		}
 	}
 
 
diff --git a/Engine/Plugins/Mutable/Source/MutableTools/Private/MuT/ASTOpLayoutPack.cpp b/Engine/Plugins/Mutable/Source/MutableTools/Private/MuT/ASTOpLayoutPack.cpp
--- a/Engine/Plugins/Mutable/Source/MutableTools/Private/MuT/ASTOpLayoutPack.cpp
+++ b/Engine/Plugins/Mutable/Source/MutableTools/Private/MuT/ASTOpLayoutPack.cpp
@@ -70,20 +70,24 @@ namespace UE::Mutable::Private
 	}
 
 
-	void ASTOpLayoutPack::Link(FProgram& program, FLinkerOptions*)
+	void ASTOpLayoutPack::Link(FProgram& Program, FLinkerOptions*)
 	{
 		// Already linked?
-		if (!linkedAddress)
+		if (!LinkedAddress)
 		{
-			OP::LayoutRemoveBlocksArgs Args;
+			FOperation::LayoutRemoveBlocksArgs Args;
 			FMemory::Memzero(Args);
 
-			if (Source) Args.Source = Source->linkedAddress;
+			if (Source)
+			{
+				Args.Source = Source->LinkedAddress;
+			}
 
-			linkedAddress = (OP::ADDRESS)program.OpAddress.Num();
-			program.OpAddress.Add((uint32_t)program.ByteCode.Num());
-			AppendCode(program.ByteCode, GetOpType());
-			AppendCode(program.ByteCode, Args);
+			++Program.NumOps;
+			LinkedAddress = MakeProgramAddress(GetOpType(), Program.ByteCode.Num());
+			
+			AppendCode(Program.ByteCode, GetOpType());
+			AppendCode(Program.ByteCode, Args);
 		}
 
 	}
diff --git a/Engine/Plugins/Mutable/Source/MutableTools/Private/MuT/ASTOpLayoutRemoveBlocks.cpp b/Engine/Plugins/Mutable/Source/MutableTools/Private/MuT/ASTOpLayoutRemoveBlocks.cpp
--- a/Engine/Plugins/Mutable/Source/MutableTools/Private/MuT/ASTOpLayoutRemoveBlocks.cpp
+++ b/Engine/Plugins/Mutable/Source/MutableTools/Private/MuT/ASTOpLayoutRemoveBlocks.cpp
@@ -74,21 +74,29 @@ namespace UE::Mutable::Private
 	}
 
 
-	void ASTOpLayoutRemoveBlocks::Link(FProgram& program, FLinkerOptions*)
+	void ASTOpLayoutRemoveBlocks::Link(FProgram& Program, FLinkerOptions*)
 	{
 		// Already linked?
-		if (!linkedAddress)
+		if (!LinkedAddress)
 		{
-			OP::LayoutRemoveBlocksArgs Args;
+			FOperation::LayoutRemoveBlocksArgs Args;
 			FMemory::Memzero(Args);
 
-			if (Source) Args.Source = Source->linkedAddress;
-			if (ReferenceLayout) Args.ReferenceLayout = ReferenceLayout->linkedAddress;
+			if (Source) 
+			{
+				Args.Source = Source->LinkedAddress;
+			}
 
-			linkedAddress = (OP::ADDRESS)program.OpAddress.Num();
-			program.OpAddress.Add((uint32_t)program.ByteCode.Num());
-			AppendCode(program.ByteCode, EOpType::LA_REMOVEBLOCKS);
-			AppendCode(program.ByteCode, Args);
+			if (ReferenceLayout) 
+			{
+				Args.ReferenceLayout = ReferenceLayout->LinkedAddress;
+			}
+
+			++Program.NumOps;
+			LinkedAddress = MakeProgramAddress(GetOpType(), Program.ByteCode.Num());
+
+			AppendCode(Program.ByteCode, GetOpType());
+			AppendCode(Program.ByteCode, Args);
 		}
 
 	}
diff --git a/Engine/Plugins/Mutable/Source/MutableTools/Private/MuT/ASTOpMaterialBreak.cpp b/Engine/Plugins/Mutable/Source/MutableTools/Private/MuT/ASTOpMaterialBreak.cpp
--- a/Engine/Plugins/Mutable/Source/MutableTools/Private/MuT/ASTOpMaterialBreak.cpp
+++ b/Engine/Plugins/Mutable/Source/MutableTools/Private/MuT/ASTOpMaterialBreak.cpp
@@ -81,23 +81,25 @@ namespace UE::Mutable::Private
 	void ASTOpMaterialBreak::Link(FProgram& Program, FLinkerOptions*)
 	{
 		// Already linked?
-		if (!linkedAddress)
+		if (!LinkedAddress)
 		{
-			linkedAddress = (OP::ADDRESS)Program.OpAddress.Num();
-			Program.OpAddress.Add((uint32)Program.ByteCode.Num());
-			AppendCode(Program.ByteCode, Type);
 
-			OP::MaterialBreakArgs Args;
+			FOperation::MaterialBreakArgs Args;
 			FMemory::Memzero(Args);
 
 			if (Material)
 			{
-				Args.Material = Material->linkedAddress;
+				Args.Material = Material->LinkedAddress;
 			}
 
 			Args.ParameterName = Program.AddConstant(ParameterKey.ParameterName);
 			Args.LayerIndex = ParameterKey.LayerIndex;
 
+
+			++Program.NumOps;
+			LinkedAddress = MakeProgramAddress(GetOpType(), Program.ByteCode.Num());
+
+			AppendCode(Program.ByteCode, GetOpType());
 			AppendCode(Program.ByteCode, Args);
 		}
 	}
diff --git a/Engine/Plugins/Mutable/Source/MutableTools/Private/MuT/ASTOpMaterialBreakImageParameter.cpp b/Engine/Plugins/Mutable/Source/MutableTools/Private/MuT/ASTOpMaterialBreakImageParameter.cpp
--- a/Engine/Plugins/Mutable/Source/MutableTools/Private/MuT/ASTOpMaterialBreakImageParameter.cpp
+++ b/Engine/Plugins/Mutable/Source/MutableTools/Private/MuT/ASTOpMaterialBreakImageParameter.cpp
@@ -71,22 +71,22 @@ namespace UE::Mutable::Private
 	void ASTOpImageFromMaterialParameter::Link(FProgram& Program, FLinkerOptions*)
 	{
 		// Already linked?
-		if (!linkedAddress)
+		if (!LinkedAddress)
 		{
 			int32 LinkedParameterIndex = Program.Parameters.Find(MaterialParameter);
 			check(LinkedParameterIndex != INDEX_NONE); // check ASTOpParameter.cpp if this happens
 
-			linkedAddress = (OP::ADDRESS)Program.OpAddress.Num();
-			Program.OpAddress.Add((uint32)Program.ByteCode.Num());
-			AppendCode(Program.ByteCode, GetOpType());
-
-			OP::MaterialBreakImageParameterArgs Args;
+			FOperation::MaterialBreakImageParameterArgs Args;
 			FMemory::Memzero(Args);
 
-			Args.MaterialParameter = OP::ADDRESS(LinkedParameterIndex);
+			Args.MaterialParameter = FOperation::ADDRESS(LinkedParameterIndex);
 			Args.ParameterName = Program.AddConstant(ParameterKey.ParameterName);
 			Args.LayerIndex = ParameterKey.LayerIndex;
 
+			++Program.NumOps;
+			LinkedAddress = MakeProgramAddress(GetOpType(), Program.ByteCode.Num());
+
+			AppendCode(Program.ByteCode, GetOpType());
 			AppendCode(Program.ByteCode, Args);
 		}
 	}
diff --git a/Engine/Plugins/Mutable/Source/MutableTools/Private/MuT/ASTOpMaterialModify.cpp b/Engine/Plugins/Mutable/Source/MutableTools/Private/MuT/ASTOpMaterialModify.cpp
--- a/Engine/Plugins/Mutable/Source/MutableTools/Private/MuT/ASTOpMaterialModify.cpp
+++ b/Engine/Plugins/Mutable/Source/MutableTools/Private/MuT/ASTOpMaterialModify.cpp
@@ -108,28 +108,29 @@ namespace UE::Mutable::Private
 	void ASTOpMaterialModify::Link(FProgram& Program, FLinkerOptions*)
 	{
 		// Already linked?
-		if (!linkedAddress)
+		if (!LinkedAddress)
 		{
-			linkedAddress = (OP::ADDRESS)Program.OpAddress.Num();
-			Program.OpAddress.Add((uint32)Program.ByteCode.Num());
-			AppendCode(Program.ByteCode, GetOpType());
-
-			OP::MaterialModifyArgs Args;
+			FOperation::MaterialModifyArgs Args;
 			FMemory::Memzero(Args);
 
 			if (Material)
 			{
-				Args.Material = Material->linkedAddress;
+				Args.Material = Material->LinkedAddress;
 			}
 
 			Args.NumParameters = ParametersToModify.Num();
+
+			++Program.NumOps;
+			LinkedAddress = MakeProgramAddress(GetOpType(), Program.ByteCode.Num());
+			
+			AppendCode(Program.ByteCode, GetOpType());
 			AppendCode(Program.ByteCode, Args);
 
 			//Store parameters info directly into the bytecode
 			for (const TPair<FParameterKey, ASTChild>& Parameter : ParametersToModify)
 			{
 				//Parameter Name
-				AppendCode(Program.ByteCode,Program.AddConstant(Parameter.Key.ParameterName));
+				AppendCode(Program.ByteCode, Program.AddConstant(Parameter.Key.ParameterName));
 				
 				// Layer Index
 				AppendCode(Program.ByteCode, Parameter.Key.LayerIndex);
@@ -137,7 +138,7 @@ namespace UE::Mutable::Private
 				// Check if exists a child operation
 				if (Parameter.Value)
 				{
-					AppendCode(Program.ByteCode, Parameter.Value->linkedAddress);
+					AppendCode(Program.ByteCode, Parameter.Value->LinkedAddress);
 
 					// Only images need this step
 					if (GetOpDataType(Parameter.Value.Child->GetOpType()) == EDataType::Image)
@@ -173,11 +174,11 @@ namespace UE::Mutable::Private
 
 						if (it != INDEX_NONE)
 						{
-							Program.RelevantParameterList.Add(Parameter.Value->linkedAddress, it);
+							Program.RelevantParameterList.Add(Parameter.Value->LinkedAddress, it);
 						}
 						else
 						{
-							Program.RelevantParameterList.Add(Parameter.Value->linkedAddress, Program.ParameterLists.Num());
+							Program.RelevantParameterList.Add(Parameter.Value->LinkedAddress, Program.ParameterLists.Num());
 							Program.ParameterLists.Add(params);
 						}
 					}
diff --git a/Engine/Plugins/Mutable/Source/MutableTools/Private/MuT/ASTOpMaterialSkeletalMeshBreak.cpp b/Engine/Plugins/Mutable/Source/MutableTools/Private/MuT/ASTOpMaterialSkeletalMeshBreak.cpp
--- a/Engine/Plugins/Mutable/Source/MutableTools/Private/MuT/ASTOpMaterialSkeletalMeshBreak.cpp
+++ b/Engine/Plugins/Mutable/Source/MutableTools/Private/MuT/ASTOpMaterialSkeletalMeshBreak.cpp
@@ -60,24 +60,25 @@ namespace UE::Mutable::Private
 	void ASTOpMaterialSkeletalMeshBreak::Link(FProgram& Program, FLinkerOptions* LinkerOptions)
 	{
 		// Already linked?
-		if (linkedAddress)
+		if (LinkedAddress)
 		{
 			return;
 		}
 		
-		OP::MaterialSkeletalMeshBreakArgs Args;
+		FOperation::MaterialSkeletalMeshBreakArgs Args;
 		FMemory::Memzero(Args);
 
 		Args.SlotNameStringIndex = Program.AddConstant(SlotName.ToString());
 		
 		if (SkeletalMesh)
 		{
-			Args.SkeletalMesh = SkeletalMesh->linkedAddress;
+			Args.SkeletalMesh = SkeletalMesh->LinkedAddress;
 		}
 
-		linkedAddress = (OP::ADDRESS)Program.OpAddress.Num();
-		Program.OpAddress.Add((uint32_t)Program.ByteCode.Num());
-		AppendCode(Program.ByteCode, this->GetOpType());
+		++Program.NumOps;
+		LinkedAddress = MakeProgramAddress(GetOpType(), Program.ByteCode.Num());
+		
+		AppendCode(Program.ByteCode, GetOpType());
 		AppendCode(Program.ByteCode, Args);
 	}
 }
diff --git a/Engine/Plugins/Mutable/Source/MutableTools/Private/MuT/ASTOpMeshAddMetadata.cpp b/Engine/Plugins/Mutable/Source/MutableTools/Private/MuT/ASTOpMeshAddMetadata.cpp
--- a/Engine/Plugins/Mutable/Source/MutableTools/Private/MuT/ASTOpMeshAddMetadata.cpp
+++ b/Engine/Plugins/Mutable/Source/MutableTools/Private/MuT/ASTOpMeshAddMetadata.cpp
@@ -85,16 +85,14 @@ namespace UE::Mutable::Private
 	void ASTOpMeshAddMetadata::Link(FProgram& Program, FLinkerOptions*)
 	{
 		// Already linked?
-		if (!linkedAddress)
-		{
-			linkedAddress = (OP::ADDRESS)Program.OpAddress.Num();
-	
-			OP::MeshAddMetadataArgs Args;
+		if (!LinkedAddress)
+		{	
+			FOperation::MeshAddMetadataArgs Args;
 			FMemory::Memzero(Args);
 		
-			Args.Source = Source ? Source->linkedAddress : 0;
+			Args.Source = Source ? Source->LinkedAddress : 0;
 			
-			using OpEnumFlags = OP::MeshAddMetadataArgs::EnumFlags;
+			using OpEnumFlags = FOperation::MeshAddMetadataArgs::EnumFlags;
 
 			Args.Flags = OpEnumFlags::None;
 			
@@ -205,8 +203,10 @@ namespace UE::Mutable::Private
 				}
 			}
 
-			Program.OpAddress.Add((uint32)Program.ByteCode.Num());
-			AppendCode(Program.ByteCode, EOpType::ME_ADDMETADATA);
+			++Program.NumOps;
+			LinkedAddress = MakeProgramAddress(GetOpType(), Program.ByteCode.Num());
+
+			AppendCode(Program.ByteCode, GetOpType());
 			AppendCode(Program.ByteCode, Args);
 		}
 	}
diff --git a/Engine/Plugins/Mutable/Source/MutableTools/Private/MuT/ASTOpMeshApplyLayout.cpp b/Engine/Plugins/Mutable/Source/MutableTools/Private/MuT/ASTOpMeshApplyLayout.cpp
--- a/Engine/Plugins/Mutable/Source/MutableTools/Private/MuT/ASTOpMeshApplyLayout.cpp
+++ b/Engine/Plugins/Mutable/Source/MutableTools/Private/MuT/ASTOpMeshApplyLayout.cpp
@@ -71,17 +71,26 @@ namespace UE::Mutable::Private
 	void ASTOpMeshApplyLayout::Link(FProgram& Program, FLinkerOptions*)
 	{
 		// Already linked?
-		if (!linkedAddress)
+		if (!LinkedAddress)
 		{
-			OP::MeshApplyLayoutArgs Args;
+			FOperation::MeshApplyLayoutArgs Args;
 			FMemory::Memzero(Args);
 
-			if (Mesh) Args.Mesh = Mesh->linkedAddress;
-			if (Layout) Args.Layout = Layout->linkedAddress;
+			if (Mesh) 
+			{
+				Args.Mesh = Mesh->LinkedAddress;
+			}
+
+			if (Layout) 
+			{
+				Args.Layout = Layout->LinkedAddress;
+			}
+
 			Args.Channel = Channel;
 
-			linkedAddress = (OP::ADDRESS)Program.OpAddress.Num();
-			Program.OpAddress.Add(Program.ByteCode.Num());
+			++Program.NumOps;
+			LinkedAddress = MakeProgramAddress(GetOpType(), Program.ByteCode.Num());
+			
 			AppendCode(Program.ByteCode, GetOpType());
 			AppendCode(Program.ByteCode, Args);
 		}
diff --git a/Engine/Plugins/Mutable/Source/MutableTools/Private/MuT/ASTOpMeshApplyPose.cpp b/Engine/Plugins/Mutable/Source/MutableTools/Private/MuT/ASTOpMeshApplyPose.cpp
--- a/Engine/Plugins/Mutable/Source/MutableTools/Private/MuT/ASTOpMeshApplyPose.cpp
+++ b/Engine/Plugins/Mutable/Source/MutableTools/Private/MuT/ASTOpMeshApplyPose.cpp
@@ -67,16 +67,24 @@ namespace UE::Mutable::Private
 	void ASTOpMeshApplyPose::Link(FProgram& Program, FLinkerOptions*)
 	{
 		// Already linked?
-		if (!linkedAddress)
+		if (!LinkedAddress)
 		{
-			OP::MeshApplyPoseArgs Args;
+			FOperation::MeshApplyPoseArgs Args;
 			FMemory::Memzero(Args);
 
-			if (Base) Args.base = Base->linkedAddress;
-			if (Pose) Args.pose = Pose->linkedAddress;
+			if (Base) 
+			{
+				Args.base = Base->LinkedAddress;
+			}
+
+			if (Pose) 
+			{
+				Args.pose = Pose->LinkedAddress;
+			}
 
-			linkedAddress = (OP::ADDRESS)Program.OpAddress.Num();
-			Program.OpAddress.Add(Program.ByteCode.Num());
+			++Program.NumOps;
+			LinkedAddress = MakeProgramAddress(GetOpType(), Program.ByteCode.Num());
+			
 			AppendCode(Program.ByteCode, GetOpType());
 			AppendCode(Program.ByteCode, Args);
 		}
diff --git a/Engine/Plugins/Mutable/Source/MutableTools/Private/MuT/ASTOpMeshApplyShape.cpp b/Engine/Plugins/Mutable/Source/MutableTools/Private/MuT/ASTOpMeshApplyShape.cpp
--- a/Engine/Plugins/Mutable/Source/MutableTools/Private/MuT/ASTOpMeshApplyShape.cpp
+++ b/Engine/Plugins/Mutable/Source/MutableTools/Private/MuT/ASTOpMeshApplyShape.cpp
@@ -87,12 +87,12 @@ namespace UE::Mutable::Private
 	}
 
 
-	void ASTOpMeshApplyShape::Link(FProgram& program, FLinkerOptions*)
+	void ASTOpMeshApplyShape::Link(FProgram& Program, FLinkerOptions*)
 	{
 		// Already linked?
-		if (!linkedAddress)
+		if (!LinkedAddress)
 		{
-			OP::MeshApplyShapeArgs Args;
+			FOperation::MeshApplyShapeArgs Args;
 			FMemory::Memzero(Args);
 
 			constexpr EMeshBindShapeFlags NoFlags = EMeshBindShapeFlags::None;
@@ -105,13 +105,14 @@ namespace UE::Mutable::Private
 
 			Args.flags = static_cast<uint32>(BindFlags);
 
-			Args.mesh = Mesh ? Mesh->linkedAddress : 0;
-			Args.shape = Shape ? Shape->linkedAddress : 0;
+			Args.mesh = Mesh ? Mesh->LinkedAddress : 0;
+			Args.shape = Shape ? Shape->LinkedAddress : 0;
 
-			linkedAddress = (OP::ADDRESS)program.OpAddress.Num();
-			program.OpAddress.Add((uint32_t)program.ByteCode.Num());
-			AppendCode(program.ByteCode, EOpType::ME_APPLYSHAPE);
-			AppendCode(program.ByteCode, Args);
+			++Program.NumOps;
+			LinkedAddress = MakeProgramAddress(GetOpType(), Program.ByteCode.Num());
+			
+			AppendCode(Program.ByteCode, GetOpType());
+			AppendCode(Program.ByteCode, Args);
 		}
 
 	}
diff --git a/Engine/Plugins/Mutable/Source/MutableTools/Private/MuT/ASTOpMeshBindShape.cpp b/Engine/Plugins/Mutable/Source/MutableTools/Private/MuT/ASTOpMeshBindShape.cpp
--- a/Engine/Plugins/Mutable/Source/MutableTools/Private/MuT/ASTOpMeshBindShape.cpp
+++ b/Engine/Plugins/Mutable/Source/MutableTools/Private/MuT/ASTOpMeshBindShape.cpp
@@ -127,12 +127,12 @@ namespace UE::Mutable::Private
 	}
 
 
-	void ASTOpMeshBindShape::Link(FProgram& program, FLinkerOptions*)
+	void ASTOpMeshBindShape::Link(FProgram& Program, FLinkerOptions*)
 	{
 		// Already linked?
-		if (!linkedAddress)
+		if (!LinkedAddress)
 		{
-			OP::MeshBindShapeArgs Args;
+			FOperation::MeshBindShapeArgs Args;
 			FMemory::Memzero(Args);
 
 			constexpr EMeshBindShapeFlags NoFlags = EMeshBindShapeFlags::None;
@@ -171,29 +171,30 @@ namespace UE::Mutable::Private
 
 			if (Mesh)
 			{
-				Args.mesh = Mesh->linkedAddress;
+				Args.mesh = Mesh->LinkedAddress;
 			}
 
 			if (Shape)
 			{
-				Args.shape = Shape->linkedAddress;
+				Args.shape = Shape->LinkedAddress;
 			}
 
-			linkedAddress = (OP::ADDRESS)program.OpAddress.Num();
-			program.OpAddress.Add((uint32_t)program.ByteCode.Num());
-			AppendCode(program.ByteCode, EOpType::ME_BINDSHAPE);
-			AppendCode(program.ByteCode, Args);
+			++Program.NumOps;
+			LinkedAddress = MakeProgramAddress(GetOpType(), Program.ByteCode.Num());
+			
+			AppendCode(Program.ByteCode, GetOpType());
+			AppendCode(Program.ByteCode, Args);
 
-			AppendCode(program.ByteCode, (int32)BonesToDeform.Num());
+			AppendCode(Program.ByteCode, (int32)BonesToDeform.Num());
 			for (FName Bone : BonesToDeform)
 			{
-				AppendCode(program.ByteCode, program.AddConstant(Bone));
+				AppendCode(Program.ByteCode, Program.AddConstant(Bone));
 			}
 
-			AppendCode(program.ByteCode, (int32)PhysicsToDeform.Num());
+			AppendCode(Program.ByteCode, (int32)PhysicsToDeform.Num());
 			for (FName PhysicsBone : PhysicsToDeform)
 			{
-				AppendCode(program.ByteCode, program.AddConstant(PhysicsBone));
+				AppendCode(Program.ByteCode, Program.AddConstant(PhysicsBone));
 			}
 		}
 	}
diff --git a/Engine/Plugins/Mutable/Source/MutableTools/Private/MuT/ASTOpMeshClipDeform.cpp b/Engine/Plugins/Mutable/Source/MutableTools/Private/MuT/ASTOpMeshClipDeform.cpp
--- a/Engine/Plugins/Mutable/Source/MutableTools/Private/MuT/ASTOpMeshClipDeform.cpp
+++ b/Engine/Plugins/Mutable/Source/MutableTools/Private/MuT/ASTOpMeshClipDeform.cpp
@@ -66,30 +66,31 @@ namespace UE::Mutable::Private
 	}
 
 
-	void ASTOpMeshClipDeform::Link(FProgram& program, FLinkerOptions*)
+	void ASTOpMeshClipDeform::Link(FProgram& Program, FLinkerOptions*)
 	{
 		// Already linked?
-		if (!linkedAddress)
+		if (!LinkedAddress)
 		{
-			OP::MeshClipDeformArgs Args;
+			FOperation::MeshClipDeformArgs Args;
 			FMemory::Memzero(Args);
 
 			Args.FaceCullStrategy = FaceCullStrategy;
 
 			if (Mesh)
 			{
-				Args.mesh = Mesh->linkedAddress;
+				Args.mesh = Mesh->LinkedAddress;
 			}
 
 			if (ClipShape)
 			{
-				Args.clipShape = ClipShape->linkedAddress;
+				Args.clipShape = ClipShape->LinkedAddress;
 			}
 
-			linkedAddress = (OP::ADDRESS)program.OpAddress.Num();
-			program.OpAddress.Add((uint32_t)program.ByteCode.Num());
-			AppendCode(program.ByteCode, EOpType::ME_CLIPDEFORM);
-			AppendCode(program.ByteCode, Args);
+			++Program.NumOps;
+			LinkedAddress = MakeProgramAddress(GetOpType(), Program.ByteCode.Num());
+			
+			AppendCode(Program.ByteCode, GetOpType());
+			AppendCode(Program.ByteCode, Args);
 		}
 
 	}
diff --git a/Engine/Plugins/Mutable/Source/MutableTools/Private/MuT/ASTOpMeshClipMorphPlane.cpp b/Engine/Plugins/Mutable/Source/MutableTools/Private/MuT/ASTOpMeshClipMorphPlane.cpp
--- a/Engine/Plugins/Mutable/Source/MutableTools/Private/MuT/ASTOpMeshClipMorphPlane.cpp
+++ b/Engine/Plugins/Mutable/Source/MutableTools/Private/MuT/ASTOpMeshClipMorphPlane.cpp
@@ -86,12 +86,15 @@ namespace UE::Mutable::Private
 	void ASTOpMeshClipMorphPlane::Link(FProgram& Program, FLinkerOptions*)
 	{
 		// Already linked?
-		if (!linkedAddress)
+		if (!LinkedAddress)
 		{
-			OP::MeshClipMorphPlaneArgs Args;
+			FOperation::MeshClipMorphPlaneArgs Args;
 			FMemory::Memzero(Args);
 
-			if (Source) Args.Source = Source->linkedAddress;
+			if (Source) 
+			{
+				Args.Source = Source->LinkedAddress;
+			}
 
 			Args.MorphShape = Program.AddConstant(MorphShape);
 
@@ -111,8 +114,9 @@ namespace UE::Mutable::Private
 			Args.Factor = Factor;
 			Args.MaxBoneRadius = VertexSelectionBoneMaxRadius;
 
-			linkedAddress = (OP::ADDRESS)Program.OpAddress.Num();
-			Program.OpAddress.Add(Program.ByteCode.Num());
+			++Program.NumOps;
+			LinkedAddress = MakeProgramAddress(GetOpType(), Program.ByteCode.Num());
+			
 			AppendCode(Program.ByteCode, GetOpType());
 			AppendCode(Program.ByteCode, Args);
 		}
