[Mutable] Program cache rewrok.

diff --git a/Engine/Plugins/Mutable/Source/MutableTools/Private/MuT/ASTOpMeshClipWithMesh.cpp b/Engine/Plugins/Mutable/Source/MutableTools/Private/MuT/ASTOpMeshClipWithMesh.cpp
--- a/Engine/Plugins/Mutable/Source/MutableTools/Private/MuT/ASTOpMeshClipWithMesh.cpp
+++ b/Engine/Plugins/Mutable/Source/MutableTools/Private/MuT/ASTOpMeshClipWithMesh.cpp
@@ -70,24 +70,26 @@ namespace UE::Mutable::Private
 	void ASTOpMeshClipWithMesh::Link(FProgram& Program, FLinkerOptions*)
 	{
 		// Already linked?
-		if (!linkedAddress)
+		if (!LinkedAddress)
 		{
-			OP::MeshClipWithMeshArgs Args;
+			FOperation::MeshClipWithMeshArgs Args;
 			FMemory::Memzero(Args);
 
 			if (Source)
 			{
-				Args.Source = Source->linkedAddress;
+				Args.Source = Source->LinkedAddress;
 			}
+
 			if (ClipMesh)
 			{
-				Args.ClipMesh = ClipMesh->linkedAddress;
+				Args.ClipMesh = ClipMesh->LinkedAddress;
 			}
 
 			Args.FaceCullStrategy = FaceCullStrategy;
+
+			++Program.NumOps;
+			LinkedAddress = MakeProgramAddress(GetOpType(), Program.ByteCode.Num());
 			
-			linkedAddress = (OP::ADDRESS)Program.OpAddress.Num();
-			Program.OpAddress.Add(Program.ByteCode.Num());
 			AppendCode(Program.ByteCode, GetOpType());
 			AppendCode(Program.ByteCode, Args);
 		}
diff --git a/Engine/Plugins/Mutable/Source/MutableTools/Private/MuT/ASTOpMeshDifference.cpp b/Engine/Plugins/Mutable/Source/MutableTools/Private/MuT/ASTOpMeshDifference.cpp
--- a/Engine/Plugins/Mutable/Source/MutableTools/Private/MuT/ASTOpMeshDifference.cpp
+++ b/Engine/Plugins/Mutable/Source/MutableTools/Private/MuT/ASTOpMeshDifference.cpp
@@ -77,30 +77,28 @@ namespace UE::Mutable::Private
 	}
 
 
-	//---------------------------------------------------------------------------------------------
-	void ASTOpMeshDifference::Link(FProgram& program, FLinkerOptions*)
+	void ASTOpMeshDifference::Link(FProgram& Program, FLinkerOptions*)
 	{
 		// Already linked?
-		if (!linkedAddress)
+		if (!LinkedAddress)
 		{
-			linkedAddress = (OP::ADDRESS)program.OpAddress.Num();
-
-			program.OpAddress.Add((uint32_t)program.ByteCode.Num());
-			AppendCode(program.ByteCode, EOpType::ME_DIFFERENCE);
+			++Program.NumOps;
+			LinkedAddress = MakeProgramAddress(GetOpType(), Program.ByteCode.Num());
+			AppendCode(Program.ByteCode, GetOpType());
 
-			OP::ADDRESS BaseAt = Base ? Base->linkedAddress : 0;
-			AppendCode(program.ByteCode, BaseAt);
+			FOperation::ADDRESS BaseAt = Base ? Base->LinkedAddress : 0;
+			AppendCode(Program.ByteCode, BaseAt);
 
-			OP::ADDRESS TargetAt = Target ? Target->linkedAddress : 0;
-			AppendCode(program.ByteCode, TargetAt);
+			FOperation::ADDRESS TargetAt = Target ? Target->LinkedAddress : 0;
+			AppendCode(Program.ByteCode, TargetAt);
 
-			AppendCode(program.ByteCode, (uint8)bIgnoreTextureCoords);
+			AppendCode(Program.ByteCode, (uint8)bIgnoreTextureCoords);
 
-			AppendCode(program.ByteCode, (uint8)Channels.Num());
+			AppendCode(Program.ByteCode, (uint8)Channels.Num());
 			for (const FChannel& b : Channels)
 			{
-				AppendCode(program.ByteCode, b.Semantic);
-				AppendCode(program.ByteCode, b.SemanticIndex);
+				AppendCode(Program.ByteCode, b.Semantic);
+				AppendCode(Program.ByteCode, b.SemanticIndex);
 			}
 		}
 	}
diff --git a/Engine/Plugins/Mutable/Source/MutableTools/Private/MuT/ASTOpMeshExtractLayoutBlocks.cpp b/Engine/Plugins/Mutable/Source/MutableTools/Private/MuT/ASTOpMeshExtractLayoutBlocks.cpp
--- a/Engine/Plugins/Mutable/Source/MutableTools/Private/MuT/ASTOpMeshExtractLayoutBlocks.cpp
+++ b/Engine/Plugins/Mutable/Source/MutableTools/Private/MuT/ASTOpMeshExtractLayoutBlocks.cpp
@@ -81,23 +81,24 @@ namespace UE::Mutable::Private
 	}
 
 
-	void ASTOpMeshExtractLayoutBlocks::Link(FProgram& program, FLinkerOptions*)
+	void ASTOpMeshExtractLayoutBlocks::Link(FProgram& Program, FLinkerOptions*)
 	{
 		// Already linked?
-		if (!linkedAddress)
+		if (!LinkedAddress)
 		{
-			linkedAddress = (OP::ADDRESS)program.OpAddress.Num();
+			++Program.NumOps;
+			LinkedAddress = MakeProgramAddress(GetOpType(), Program.ByteCode.Num());	
+			
+			AppendCode(Program.ByteCode, GetOpType());
 
-			program.OpAddress.Add((uint32)program.ByteCode.Num());
-			AppendCode(program.ByteCode, EOpType::ME_EXTRACTLAYOUTBLOCK);
-			OP::ADDRESS sourceAt = Source ? Source->linkedAddress : 0;
-			AppendCode(program.ByteCode, sourceAt);
-			AppendCode(program.ByteCode, (uint16)LayoutIndex);
-			AppendCode(program.ByteCode, (uint16)Blocks.Num());
+			FOperation::ADDRESS sourceAt = Source ? Source->LinkedAddress : 0;
+			AppendCode(Program.ByteCode, sourceAt);
+			AppendCode(Program.ByteCode, (uint16)LayoutIndex);
+			AppendCode(Program.ByteCode, (uint16)Blocks.Num());
 
 			for (uint64 Id : Blocks)
 			{
-				AppendCode(program.ByteCode, Id);
+				AppendCode(Program.ByteCode, Id);
 			}
 		}
 	}
diff --git a/Engine/Plugins/Mutable/Source/MutableTools/Private/MuT/ASTOpMeshFormat.cpp b/Engine/Plugins/Mutable/Source/MutableTools/Private/MuT/ASTOpMeshFormat.cpp
--- a/Engine/Plugins/Mutable/Source/MutableTools/Private/MuT/ASTOpMeshFormat.cpp
+++ b/Engine/Plugins/Mutable/Source/MutableTools/Private/MuT/ASTOpMeshFormat.cpp
@@ -85,27 +85,35 @@ void ASTOpMeshFormat::ForEachChild(const TFunctionRef<void(ASTChild&)> f )
 }
 
 
-void ASTOpMeshFormat::Link( FProgram& program, FLinkerOptions* )
+void ASTOpMeshFormat::Link( FProgram& Program, FLinkerOptions* )
 {
     // Already linked?
-    if (!linkedAddress)
+    if (!LinkedAddress)
     {
-        OP::MeshFormatArgs Args;
+        FOperation::MeshFormatArgs Args;
 		FMemory::Memzero(Args);
 
 		Args.Flags = Flags;
 		if (bOptimizeBuffers)
 		{
-			Args.Flags = Args.Flags | OP::MeshFormatArgs::OptimizeBuffers;
+			Args.Flags = Args.Flags | FOperation::MeshFormatArgs::OptimizeBuffers;
 		}
 
-		if (Source) Args.source = Source->linkedAddress;
-		if (Format) Args.format = Format->linkedAddress;
+		if (Source) 
+		{
+			Args.source = Source->LinkedAddress;
+		}
+
+		if (Format) 
+		{
+			Args.format = Format->LinkedAddress;
+		}
 
-        linkedAddress = (OP::ADDRESS)program.OpAddress.Num();
-        program.OpAddress.Add((uint32)program.ByteCode.Num());
-        AppendCode(program.ByteCode,EOpType::ME_FORMAT);
-        AppendCode(program.ByteCode,Args);
+        ++Program.NumOps;
+        LinkedAddress = MakeProgramAddress(GetOpType(), Program.ByteCode.Num());
+        
+        AppendCode(Program.ByteCode, GetOpType());
+        AppendCode(Program.ByteCode, Args);
     }
 
 }
@@ -313,7 +321,7 @@ UE::Mutable::Private::Ptr<ASTOp> Sink_MeshFormatAST::Visit(const UE::Mutable::Pr
 		NewFormatConstant->SetValue(TargetFormat);
 		NewFormatConstant->SourceDataDescriptor = at->GetSourceDataDescriptor();
 	
-        NewFormat->Flags = NewFormat->Flags | OP::MeshFormatArgs::OptimizeBuffers;
+        NewFormat->Flags = NewFormat->Flags | FOperation::MeshFormatArgs::OptimizeBuffers;
 
         // TODO: Optimize, in case no skinning data is found in the format mesh a generic buffer that can represent
         // all possible skinning formats is added. This is not optimal, we may want to add a flag to the format op
diff --git a/Engine/Plugins/Mutable/Source/MutableTools/Private/MuT/ASTOpMeshMaskClipMesh.cpp b/Engine/Plugins/Mutable/Source/MutableTools/Private/MuT/ASTOpMeshMaskClipMesh.cpp
--- a/Engine/Plugins/Mutable/Source/MutableTools/Private/MuT/ASTOpMeshMaskClipMesh.cpp
+++ b/Engine/Plugins/Mutable/Source/MutableTools/Private/MuT/ASTOpMeshMaskClipMesh.cpp
@@ -77,23 +77,29 @@ namespace UE::Mutable::Private
 	}
 
 
-	//-------------------------------------------------------------------------------------------------
-	void ASTOpMeshMaskClipMesh::Link(FProgram& program, FLinkerOptions*)
+	void ASTOpMeshMaskClipMesh::Link(FProgram& Program, FLinkerOptions*)
 	{
 		// Already linked?
-		if (!linkedAddress)
+		if (!LinkedAddress)
 		{
-			OP::MeshMaskClipMeshArgs Args;
+			FOperation::MeshMaskClipMeshArgs Args;
 			FMemory::Memzero(Args);
 
-			if (source) Args.source = source->linkedAddress;
-			if (clip) Args.clip = clip->linkedAddress;
+			if (source) 
+			{
+				Args.source = source->LinkedAddress;
+			}
+
+			if (clip) 
+			{
+				Args.clip = clip->LinkedAddress;
+			}
 
-			linkedAddress = (OP::ADDRESS)program.OpAddress.Num();
-			//program.m_code.push_back(op);
-			program.OpAddress.Add((uint32_t)program.ByteCode.Num());
-			AppendCode(program.ByteCode, EOpType::ME_MASKCLIPMESH);
-			AppendCode(program.ByteCode, Args);
+			++Program.NumOps;
+			LinkedAddress = MakeProgramAddress(GetOpType(), Program.ByteCode.Num());
+			
+			AppendCode(Program.ByteCode, GetOpType());
+			AppendCode(Program.ByteCode, Args);
 		}
 	}
 
diff --git a/Engine/Plugins/Mutable/Source/MutableTools/Private/MuT/ASTOpMeshMaskClipUVMask.cpp b/Engine/Plugins/Mutable/Source/MutableTools/Private/MuT/ASTOpMeshMaskClipUVMask.cpp
--- a/Engine/Plugins/Mutable/Source/MutableTools/Private/MuT/ASTOpMeshMaskClipUVMask.cpp
+++ b/Engine/Plugins/Mutable/Source/MutableTools/Private/MuT/ASTOpMeshMaskClipUVMask.cpp
@@ -90,25 +90,41 @@ namespace UE::Mutable::Private
 	}
 
 
-	//-------------------------------------------------------------------------------------------------
-	void ASTOpMeshMaskClipUVMask::Link(FProgram& program, FLinkerOptions*)
+	void ASTOpMeshMaskClipUVMask::Link(FProgram& Program, FLinkerOptions*)
 	{
 		// Already linked?
-		if (!linkedAddress)
+		if (!LinkedAddress)
 		{
-			OP::MeshMaskClipUVMaskArgs Args;
+			FOperation::MeshMaskClipUVMaskArgs Args;
 			FMemory::Memzero(Args);
 
-			if (Source) Args.Source = Source->linkedAddress;
-			if (UVSource) Args.UVSource = UVSource->linkedAddress;
-			if (MaskImage) Args.MaskImage = MaskImage->linkedAddress;
-			if (MaskLayout) Args.MaskLayout = MaskLayout->linkedAddress;
+			if (Source) 
+			{
+				Args.Source = Source->LinkedAddress;
+			}
+
+			if (UVSource) 
+			{
+				Args.UVSource = UVSource->LinkedAddress;
+			}
+
+			if (MaskImage) 
+			{
+				Args.MaskImage = MaskImage->LinkedAddress;
+			}
+
+			if (MaskLayout) 
+			{
+				Args.MaskLayout = MaskLayout->LinkedAddress;
+			}
+
 			Args.LayoutIndex = LayoutIndex;
 
-			linkedAddress = (OP::ADDRESS)program.OpAddress.Num();
-			program.OpAddress.Add((uint32)program.ByteCode.Num());
-			AppendCode(program.ByteCode, EOpType::ME_MASKCLIPUVMASK);
-			AppendCode(program.ByteCode, Args);
+			++Program.NumOps;
+			LinkedAddress = MakeProgramAddress(GetOpType(), Program.ByteCode.Num());
+			
+			AppendCode(Program.ByteCode, GetOpType());
+			AppendCode(Program.ByteCode, Args);
 		}
 	}
 
diff --git a/Engine/Plugins/Mutable/Source/MutableTools/Private/MuT/ASTOpMeshMaskDiff.cpp b/Engine/Plugins/Mutable/Source/MutableTools/Private/MuT/ASTOpMeshMaskDiff.cpp
--- a/Engine/Plugins/Mutable/Source/MutableTools/Private/MuT/ASTOpMeshMaskDiff.cpp
+++ b/Engine/Plugins/Mutable/Source/MutableTools/Private/MuT/ASTOpMeshMaskDiff.cpp
@@ -68,16 +68,24 @@ namespace UE::Mutable::Private
 	void ASTOpMeshMaskDiff::Link(FProgram& Program, FLinkerOptions*)
 	{
 		// Already linked?
-		if (!linkedAddress)
+		if (!LinkedAddress)
 		{
-			OP::MeshMaskDiffArgs Args;
+			FOperation::MeshMaskDiffArgs Args;
 			FMemory::Memzero(Args);
 
-			if (Source) Args.Source = Source->linkedAddress;
-			if (Fragment) Args.Fragment = Fragment->linkedAddress;
+			if (Source) 
+			{
+				Args.Source = Source->LinkedAddress;
+			}
+
+			if (Fragment) 
+			{
+				Args.Fragment = Fragment->LinkedAddress;
+			}
 
-			linkedAddress = (OP::ADDRESS)Program.OpAddress.Num();
-			Program.OpAddress.Add(Program.ByteCode.Num());
+			++Program.NumOps;
+			LinkedAddress = MakeProgramAddress(GetOpType(), Program.ByteCode.Num());
+			
 			AppendCode(Program.ByteCode, GetOpType());
 			AppendCode(Program.ByteCode, Args);
 		}
diff --git a/Engine/Plugins/Mutable/Source/MutableTools/Private/MuT/ASTOpMeshMerge.cpp b/Engine/Plugins/Mutable/Source/MutableTools/Private/MuT/ASTOpMeshMerge.cpp
--- a/Engine/Plugins/Mutable/Source/MutableTools/Private/MuT/ASTOpMeshMerge.cpp
+++ b/Engine/Plugins/Mutable/Source/MutableTools/Private/MuT/ASTOpMeshMerge.cpp
@@ -69,17 +69,26 @@ namespace UE::Mutable::Private
 	void ASTOpMeshMerge::Link(FProgram& Program, FLinkerOptions*)
 	{
 		// Already linked?
-		if (!linkedAddress)
+		if (!LinkedAddress)
 		{
-			OP::MeshMergeArgs Args;
+			FOperation::MeshMergeArgs Args;
 			FMemory::Memzero(Args);
 
-			if (Base) Args.Base = Base->linkedAddress;
-			if (Added) Args.Added = Added->linkedAddress;
+			if (Base) 
+			{
+				Args.Base = Base->LinkedAddress;
+			}
+
+			if (Added) 
+			{
+				Args.Added = Added->LinkedAddress;
+			}
+
 			Args.NewSurfaceID = NewSurfaceID;
 
-			linkedAddress = (OP::ADDRESS)Program.OpAddress.Num();
-			Program.OpAddress.Add(Program.ByteCode.Num());
+			++Program.NumOps;
+			LinkedAddress = MakeProgramAddress(GetOpType(), Program.ByteCode.Num());
+
 			AppendCode(Program.ByteCode, GetOpType());
 			AppendCode(Program.ByteCode, Args);
 		}
diff --git a/Engine/Plugins/Mutable/Source/MutableTools/Private/MuT/ASTOpMeshMorph.cpp b/Engine/Plugins/Mutable/Source/MutableTools/Private/MuT/ASTOpMeshMorph.cpp
--- a/Engine/Plugins/Mutable/Source/MutableTools/Private/MuT/ASTOpMeshMorph.cpp
+++ b/Engine/Plugins/Mutable/Source/MutableTools/Private/MuT/ASTOpMeshMorph.cpp
@@ -75,18 +75,17 @@ namespace UE::Mutable::Private
 	void ASTOpMeshMorph::Link(FProgram& Program, FLinkerOptions*)
 	{
 		// Already linked?
-		if (!linkedAddress)
+		if (!LinkedAddress)
 		{
-			OP::MeshMorphArgs Args;
+			FOperation::MeshMorphArgs Args;
 			FMemory::Memzero(Args);
 
-			linkedAddress = (OP::ADDRESS)Program.OpAddress.Num();
 
 			Args.Name = Program.AddConstant(Name.ToString());
-			Args.Base = Base ? Base->linkedAddress : 0;
-			Args.Factor = Factor ? Factor->linkedAddress : 0;
+			Args.Base = Base ? Base->LinkedAddress : 0;
+			Args.Factor = Factor ? Factor->LinkedAddress : 0;
 
-			Program.OpAddress.Add((uint32)Program.ByteCode.Num());
+			LinkedAddress = MakeProgramAddress(GetOpType(), Program.ByteCode.Num());
 			AppendCode(Program.ByteCode, GetOpType());
 			AppendCode(Program.ByteCode, Args);
 			
diff --git a/Engine/Plugins/Mutable/Source/MutableTools/Private/MuT/ASTOpMeshMorphReshape.cpp b/Engine/Plugins/Mutable/Source/MutableTools/Private/MuT/ASTOpMeshMorphReshape.cpp
--- a/Engine/Plugins/Mutable/Source/MutableTools/Private/MuT/ASTOpMeshMorphReshape.cpp
+++ b/Engine/Plugins/Mutable/Source/MutableTools/Private/MuT/ASTOpMeshMorphReshape.cpp
@@ -65,28 +65,29 @@ namespace UE::Mutable::Private
 	}
 
 
-	void ASTOpMeshMorphReshape::Link(FProgram& program, FLinkerOptions*)
+	void ASTOpMeshMorphReshape::Link(FProgram& Program, FLinkerOptions*)
 	{
 		// Already linked?
-		if (!linkedAddress)
+		if (!LinkedAddress)
 		{
-			OP::MeshMorphReshapeArgs Args;
+			FOperation::MeshMorphReshapeArgs Args;
 			FMemory::Memzero(Args);
 
 			if (Morph)
 			{
-				Args.Morph = Morph->linkedAddress;
+				Args.Morph = Morph->LinkedAddress;
 			}
 
 			if (Reshape)
 			{
-				Args.Reshape = Reshape->linkedAddress;
+				Args.Reshape = Reshape->LinkedAddress;
 			}
 
-			linkedAddress = (OP::ADDRESS)program.OpAddress.Num();
-			program.OpAddress.Add((uint32_t)program.ByteCode.Num());
-			AppendCode(program.ByteCode, EOpType::ME_MORPHRESHAPE);
-			AppendCode(program.ByteCode, Args);
+			++Program.NumOps;
+			LinkedAddress = MakeProgramAddress(GetOpType(), Program.ByteCode.Num());
+
+			AppendCode(Program.ByteCode, GetOpType());
+			AppendCode(Program.ByteCode, Args);
 		}
 
 	}
diff --git a/Engine/Plugins/Mutable/Source/MutableTools/Private/MuT/ASTOpMeshOptimizeSkinning.cpp b/Engine/Plugins/Mutable/Source/MutableTools/Private/MuT/ASTOpMeshOptimizeSkinning.cpp
--- a/Engine/Plugins/Mutable/Source/MutableTools/Private/MuT/ASTOpMeshOptimizeSkinning.cpp
+++ b/Engine/Plugins/Mutable/Source/MutableTools/Private/MuT/ASTOpMeshOptimizeSkinning.cpp
@@ -67,19 +67,20 @@ namespace UE::Mutable::Private
 	void ASTOpMeshOptimizeSkinning::Link(FProgram& Program, FLinkerOptions*)
 	{
 		// Already linked?
-		if (!linkedAddress)
+		if (!LinkedAddress)
 		{
-			OP::MeshOptimizeSkinningArgs Args;
+			FOperation::MeshOptimizeSkinningArgs Args;
 			FMemory::Memzero(Args);
 
 			if (Source)
 			{
-				Args.source = Source->linkedAddress;
+				Args.source = Source->LinkedAddress;
 			}
 
-			linkedAddress = OP::ADDRESS(Program.OpAddress.Num());
-			Program.OpAddress.Add((uint32)Program.ByteCode.Num());
-			AppendCode(Program.ByteCode, EOpType::ME_OPTIMIZESKINNING);
+			++Program.NumOps;
+			LinkedAddress = MakeProgramAddress(GetOpType(), Program.ByteCode.Num());
+			
+			AppendCode(Program.ByteCode, GetOpType());
 			AppendCode(Program.ByteCode, Args);
 		}
 	}
@@ -102,11 +103,6 @@ namespace UE::Mutable::Private
 	}
 
 
-
-
-	//---------------------------------------------------------------------------------------------
-	//---------------------------------------------------------------------------------------------
-	//---------------------------------------------------------------------------------------------
 	UE::Mutable::Private::Ptr<ASTOp> Sink_MeshOptimizeSkinningAST::Apply(const ASTOpMeshOptimizeSkinning* InRoot)
 	{
 		Root = InRoot;
diff --git a/Engine/Plugins/Mutable/Source/MutableTools/Private/MuT/ASTOpMeshPrepareLayout.cpp b/Engine/Plugins/Mutable/Source/MutableTools/Private/MuT/ASTOpMeshPrepareLayout.cpp
--- a/Engine/Plugins/Mutable/Source/MutableTools/Private/MuT/ASTOpMeshPrepareLayout.cpp
+++ b/Engine/Plugins/Mutable/Source/MutableTools/Private/MuT/ASTOpMeshPrepareLayout.cpp
@@ -78,21 +78,30 @@ namespace UE::Mutable::Private
 	void ASTOpMeshPrepareLayout::Link(FProgram& Program, FLinkerOptions*)
 	{
 		// Already linked?
-		if (!linkedAddress)
+		if (!LinkedAddress)
 		{
-			OP::MeshPrepareLayoutArgs Args;
+			FOperation::MeshPrepareLayoutArgs Args;
 			FMemory::Memzero(Args);
 
-			if (Mesh) Args.Mesh = Mesh->linkedAddress;
-			if (Layout) Args.Layout = Layout->linkedAddress;
+			if (Mesh) 
+			{
+				Args.Mesh = Mesh->LinkedAddress;
+			}
+
+			if (Layout)
+			{
+				Args.Layout = Layout->LinkedAddress;
+			}
+
 			Args.LayoutChannel = LayoutChannel;
 			Args.bUseAbsoluteBlockIds = bUseAbsoluteBlockIds;
 			Args.bNormalizeUVs = bNormalizeUVs;
 			Args.bClampUVIslands = bClampUVIslands;
 			Args.bEnsureAllVerticesHaveLayoutBlock = bEnsureAllVerticesHaveLayoutBlock;
 
-			linkedAddress = (OP::ADDRESS)Program.OpAddress.Num();
-			Program.OpAddress.Add(Program.ByteCode.Num());
+			++Program.NumOps;
+			LinkedAddress = MakeProgramAddress(GetOpType(), Program.ByteCode.Num());
+			
 			AppendCode(Program.ByteCode, GetOpType());
 			AppendCode(Program.ByteCode, Args);
 		}
diff --git a/Engine/Plugins/Mutable/Source/MutableTools/Private/MuT/ASTOpMeshProject.cpp b/Engine/Plugins/Mutable/Source/MutableTools/Private/MuT/ASTOpMeshProject.cpp
--- a/Engine/Plugins/Mutable/Source/MutableTools/Private/MuT/ASTOpMeshProject.cpp
+++ b/Engine/Plugins/Mutable/Source/MutableTools/Private/MuT/ASTOpMeshProject.cpp
@@ -69,16 +69,24 @@ namespace UE::Mutable::Private
 	void ASTOpMeshProject::Link(FProgram& Program, FLinkerOptions*)
 	{
 		// Already linked?
-		if (!linkedAddress)
+		if (!LinkedAddress)
 		{
-			OP::MeshProjectArgs Args;
+			FOperation::MeshProjectArgs Args;
 			FMemory::Memzero(Args);
 
-			if (Mesh) Args.Mesh = Mesh->linkedAddress;
-			if (Projector) Args.Projector = Projector->linkedAddress;
+			if (Mesh) 
+			{
+				Args.Mesh = Mesh->LinkedAddress;
+			}
+
+			if (Projector)
+			{
+				Args.Projector = Projector->LinkedAddress;
+			}
+
+			++Program.NumOps;
+			LinkedAddress = MakeProgramAddress(GetOpType(), Program.ByteCode.Num());
 
-			linkedAddress = (OP::ADDRESS)Program.OpAddress.Num();
-			Program.OpAddress.Add(Program.ByteCode.Num());
 			AppendCode(Program.ByteCode, GetOpType());
 			AppendCode(Program.ByteCode, Args);
 		}
diff --git a/Engine/Plugins/Mutable/Source/MutableTools/Private/MuT/ASTOpMeshRemoveMask.cpp b/Engine/Plugins/Mutable/Source/MutableTools/Private/MuT/ASTOpMeshRemoveMask.cpp
--- a/Engine/Plugins/Mutable/Source/MutableTools/Private/MuT/ASTOpMeshRemoveMask.cpp
+++ b/Engine/Plugins/Mutable/Source/MutableTools/Private/MuT/ASTOpMeshRemoveMask.cpp
@@ -92,29 +92,29 @@ namespace UE::Mutable::Private
 
 
 	//---------------------------------------------------------------------------------------------
-	void ASTOpMeshRemoveMask::Link(FProgram& program, FLinkerOptions*)
+	void ASTOpMeshRemoveMask::Link(FProgram& Program, FLinkerOptions*)
 	{
 		// Already linked?
-		if (!linkedAddress)
+		if (!LinkedAddress)
 		{
-			linkedAddress = (OP::ADDRESS)program.OpAddress.Num();
+			++Program.NumOps;
+			LinkedAddress = MakeProgramAddress(GetOpType(), Program.ByteCode.Num());
 
-			program.OpAddress.Add((uint32_t)program.ByteCode.Num());
-			AppendCode(program.ByteCode, EOpType::ME_REMOVEMASK);
+			AppendCode(Program.ByteCode, GetOpType());
 
-			OP::ADDRESS sourceAt = source ? source->linkedAddress : 0;
-			AppendCode(program.ByteCode, sourceAt);
+			FOperation::ADDRESS sourceAt = source ? source->LinkedAddress : 0;
+			AppendCode(Program.ByteCode, sourceAt);
 
-			AppendCode(program.ByteCode, FaceCullStrategy);
+			AppendCode(Program.ByteCode, FaceCullStrategy);
 
-			AppendCode(program.ByteCode, (uint16)removes.Num());
+			AppendCode(Program.ByteCode, (uint16)removes.Num());
 			for (const TPair<ASTChild, ASTChild>& b : removes)
 			{
-				OP::ADDRESS condition = b.Key ? b.Key->linkedAddress : 0;
-				AppendCode(program.ByteCode, condition);
+				FOperation::ADDRESS condition = b.Key ? b.Key->LinkedAddress : 0;
+				AppendCode(Program.ByteCode, condition);
 
-				OP::ADDRESS remove = b.Value ? b.Value->linkedAddress : 0;
-				AppendCode(program.ByteCode, remove);
+				FOperation::ADDRESS remove = b.Value ? b.Value->LinkedAddress : 0;
+				AppendCode(Program.ByteCode, remove);
 			}
 		}
 	}
diff --git a/Engine/Plugins/Mutable/Source/MutableTools/Private/MuT/ASTOpMeshSkeletalMeshBreak.cpp b/Engine/Plugins/Mutable/Source/MutableTools/Private/MuT/ASTOpMeshSkeletalMeshBreak.cpp
--- a/Engine/Plugins/Mutable/Source/MutableTools/Private/MuT/ASTOpMeshSkeletalMeshBreak.cpp
+++ b/Engine/Plugins/Mutable/Source/MutableTools/Private/MuT/ASTOpMeshSkeletalMeshBreak.cpp
@@ -85,23 +85,24 @@ namespace UE::Mutable::Private
 
 	void ASTOpMeshSkeletalMeshBreak::Link(FProgram& Program, FLinkerOptions*)
 	{
-		if (!linkedAddress)
+		if (!LinkedAddress)
 		{
-			linkedAddress = static_cast<OP::ADDRESS>(Program.OpAddress.Num());
-			Program.OpAddress.Add(static_cast<uint32>(Program.ByteCode.Num()));
-			AppendCode(Program.ByteCode, Type);
 
-			OP::MeshSkeletalMeshBreakArgs Args;
+			FOperation::MeshSkeletalMeshBreakArgs Args;
 			FMemory::Memzero(Args);
 
-			check(SkeletalMeshParameter->linkedAddress);
-			Args.SkeletalMeshParameter = SkeletalMeshParameter->linkedAddress;
+			check(SkeletalMeshParameter->LinkedAddress);
+			Args.SkeletalMeshParameter = SkeletalMeshParameter->LinkedAddress;
 
 			Args.LOD = LODIndex;
 			Args.Section = SectionIndex;
 			Args.Flags = ConversionFlags;
 			Args.MeshID = MeshID;
 
+			++Program.NumOps;
+			LinkedAddress = MakeProgramAddress(Type, Program.ByteCode.Num());
+
+			AppendCode(Program.ByteCode, Type);
 			AppendCode(Program.ByteCode, Args);
 		}
 	}
diff --git a/Engine/Plugins/Mutable/Source/MutableTools/Private/MuT/ASTOpMeshTransform.cpp b/Engine/Plugins/Mutable/Source/MutableTools/Private/MuT/ASTOpMeshTransform.cpp
--- a/Engine/Plugins/Mutable/Source/MutableTools/Private/MuT/ASTOpMeshTransform.cpp
+++ b/Engine/Plugins/Mutable/Source/MutableTools/Private/MuT/ASTOpMeshTransform.cpp
@@ -62,23 +62,26 @@ namespace UE::Mutable::Private
 	}
 
 
-	void ASTOpMeshTransform::Link(FProgram& program, FLinkerOptions*)
+	void ASTOpMeshTransform::Link(FProgram& Program, FLinkerOptions*)
 	{
 		// Already linked?
-		if (!linkedAddress)
+		if (!LinkedAddress)
 		{
-			OP::MeshTransformArgs Args;
+			FOperation::MeshTransformArgs Args;
 			FMemory::Memzero(Args);
 
-			if (source) Args.source = source->linkedAddress;
+			if (source) 
+			{
+				Args.source = source->LinkedAddress;
+			}
 
-			Args.matrix = program.AddConstant(matrix);
+			Args.matrix = Program.AddConstant(matrix);
 
-			linkedAddress = (OP::ADDRESS)program.OpAddress.Num();
-			//program.m_code.push_back(op);
-			program.OpAddress.Add((uint32_t)program.ByteCode.Num());
-			AppendCode(program.ByteCode, EOpType::ME_TRANSFORM);
-			AppendCode(program.ByteCode, Args);
+			++Program.NumOps;
+			LinkedAddress = MakeProgramAddress(GetOpType(), Program.ByteCode.Num());
+
+			AppendCode(Program.ByteCode, GetOpType());
+			AppendCode(Program.ByteCode, Args);
 		}
 
 	}
diff --git a/Engine/Plugins/Mutable/Source/MutableTools/Private/MuT/ASTOpMeshTransformWithBone.cpp b/Engine/Plugins/Mutable/Source/MutableTools/Private/MuT/ASTOpMeshTransformWithBone.cpp
--- a/Engine/Plugins/Mutable/Source/MutableTools/Private/MuT/ASTOpMeshTransformWithBone.cpp
+++ b/Engine/Plugins/Mutable/Source/MutableTools/Private/MuT/ASTOpMeshTransformWithBone.cpp
@@ -60,26 +60,27 @@ void UE::Mutable::Private::ASTOpMeshTransformWithBone::ForEachChild(const TFunct
 
 void UE::Mutable::Private::ASTOpMeshTransformWithBone::Link(FProgram& Program, FLinkerOptions* Options)
 {
-	if (!linkedAddress)
+	if (!LinkedAddress)
 	{
-		OP::MeshTransformWithBoneArgs Args;
+		FOperation::MeshTransformWithBoneArgs Args;
 		FMemory::Memzero(Args);
 		if (SourceMesh)
 		{
-			Args.SourceMesh = SourceMesh->linkedAddress;
+			Args.SourceMesh = SourceMesh->LinkedAddress;
 		}
 
 		if (Matrix)
 		{
-			Args.Matrix = Matrix->linkedAddress;
+			Args.Matrix = Matrix->LinkedAddress;
 		}
 
 		Args.BoneId = Program.AddConstant(BoneName);
 		Args.ThresholdFactor = ThresholdFactor;
 
-		linkedAddress = static_cast<OP::ADDRESS>(Program.OpAddress.Num());
-		Program.OpAddress.Add(Program.ByteCode.Num());
-		AppendCode(Program.ByteCode, EOpType::ME_TRANSFORMWITHBONE);
+		++Program.NumOps;
+		LinkedAddress = MakeProgramAddress(GetOpType(), Program.ByteCode.Num());
+		
+		AppendCode(Program.ByteCode, GetOpType());
 		AppendCode(Program.ByteCode, Args);
 	}
 }
diff --git a/Engine/Plugins/Mutable/Source/MutableTools/Private/MuT/ASTOpMeshTransformWithBoundingMesh.cpp b/Engine/Plugins/Mutable/Source/MutableTools/Private/MuT/ASTOpMeshTransformWithBoundingMesh.cpp
--- a/Engine/Plugins/Mutable/Source/MutableTools/Private/MuT/ASTOpMeshTransformWithBoundingMesh.cpp
+++ b/Engine/Plugins/Mutable/Source/MutableTools/Private/MuT/ASTOpMeshTransformWithBoundingMesh.cpp
@@ -55,28 +55,31 @@ void UE::Mutable::Private::ASTOpMeshTransformWithBoundingMesh::ForEachChild(cons
 	f(matrix);
 }
 
-void UE::Mutable::Private::ASTOpMeshTransformWithBoundingMesh::Link(FProgram& program, FLinkerOptions* Options)
+void UE::Mutable::Private::ASTOpMeshTransformWithBoundingMesh::Link(FProgram& Program, FLinkerOptions* Options)
 {
-	if (!linkedAddress)
+	if (!LinkedAddress)
 	{
-		OP::MeshTransformWithinMeshArgs Args;
+		FOperation::MeshTransformWithinMeshArgs Args;
 		FMemory::Memzero(Args);
+
 		if (source)
 		{
-			Args.sourceMesh = source->linkedAddress;
+			Args.sourceMesh = source->LinkedAddress;
 		}
 		if (boundingMesh)
 		{
-			Args.boundingMesh = boundingMesh->linkedAddress;
+			Args.boundingMesh = boundingMesh->LinkedAddress;
 		}
 		if (matrix)
 		{
-			Args.matrix = matrix->linkedAddress;
+			Args.matrix = matrix->LinkedAddress;
 		}
-		linkedAddress = static_cast<OP::ADDRESS>(program.OpAddress.Num());
-		program.OpAddress.Add(program.ByteCode.Num());
-		AppendCode(program.ByteCode, EOpType::ME_TRANSFORMWITHMESH);
-		AppendCode(program.ByteCode, Args);
+	
+		++Program.NumOps;
+		LinkedAddress = MakeProgramAddress(GetOpType(), Program.ByteCode.Num());
+
+		AppendCode(Program.ByteCode, GetOpType());
+		AppendCode(Program.ByteCode, Args);
 	}
 }
 
diff --git a/Engine/Plugins/Mutable/Source/MutableTools/Private/MuT/ASTOpParameter.cpp b/Engine/Plugins/Mutable/Source/MutableTools/Private/MuT/ASTOpParameter.cpp
--- a/Engine/Plugins/Mutable/Source/MutableTools/Private/MuT/ASTOpParameter.cpp
+++ b/Engine/Plugins/Mutable/Source/MutableTools/Private/MuT/ASTOpParameter.cpp
@@ -93,24 +93,24 @@ namespace UE::Mutable::Private
 	void ASTOpParameter::Link(FProgram& Program, FLinkerOptions*)
 	{
 		// Already linked?
-		if (!linkedAddress)
+		if (!LinkedAddress)
 		{
-
 			LinkedParameterIndex = Program.Parameters.Find(Parameter);
 
 			// If this fails, it means an ASTOpParameter was created at code generation time, but not registered into the
 			// parameters map in the CodeGenerator_FirstPass.
 			check(LinkedParameterIndex != INDEX_NONE);
 
-			linkedAddress = (OP::ADDRESS)Program.OpAddress.Num();
-			Program.OpAddress.Add((uint32)Program.ByteCode.Num());
+			++Program.NumOps;
+			LinkedAddress = MakeProgramAddress(Type, Program.ByteCode.Num());
+			
 			AppendCode(Program.ByteCode, Type);
 		
 			if (Type == EOpType::MI_PARAMETER)
 			{
-				OP::MaterialParameterArgs Args;
+				FOperation::MaterialParameterArgs Args;
 				FMemory::Memzero(Args);
-				Args.Variable = (OP::ADDRESS)LinkedParameterIndex;
+				Args.Variable = (FOperation::ADDRESS)LinkedParameterIndex;
 				
 				// Images
 				TArray<FName> ImageParameterNames;
@@ -118,7 +118,7 @@ namespace UE::Mutable::Private
 				for (const TPair<FParameterKey, ASTChild>& Element : ImageOperations)
 				{
 					ImageParameterNames.Add(Element.Key.ParameterName);
-					ImageParameterOperations.Add(Element.Value->linkedAddress);
+					ImageParameterOperations.Add(Element.Value->LinkedAddress);
 		 		}
 
 				Args.ImageParameterNames = Program.AddConstant(ImageParameterNames);
@@ -127,7 +127,7 @@ namespace UE::Mutable::Private
 				// Ranges
 				for (const FRangeData& d : Ranges)
 				{
-					OP::ADDRESS sizeAt = 0;
+					FOperation::ADDRESS sizeAt = 0;
 					uint16 rangeId = 0;
 					LinkRange(Program, d, sizeAt, rangeId);
 					Program.Parameters[LinkedParameterIndex].Ranges.Add(rangeId);
@@ -143,13 +143,13 @@ namespace UE::Mutable::Private
 
 			else
 			{
-				OP::ParameterArgs Args;
+				FOperation::ParameterArgs Args;
 				FMemory::Memzero(Args);
-				Args.variable = (OP::ADDRESS)LinkedParameterIndex;
+				Args.variable = (FOperation::ADDRESS)LinkedParameterIndex;
 
 				for (const FRangeData& d : Ranges)
 				{
-					OP::ADDRESS sizeAt = 0;
+					FOperation::ADDRESS sizeAt = 0;
 					uint16 rangeId = 0;
 					LinkRange(Program, d, sizeAt, rangeId);
 					Program.Parameters[LinkedParameterIndex].Ranges.Add(rangeId);
diff --git a/Engine/Plugins/Mutable/Source/MutableTools/Private/MuT/ASTOpReferenceResource.cpp b/Engine/Plugins/Mutable/Source/MutableTools/Private/MuT/ASTOpReferenceResource.cpp
--- a/Engine/Plugins/Mutable/Source/MutableTools/Private/MuT/ASTOpReferenceResource.cpp
+++ b/Engine/Plugins/Mutable/Source/MutableTools/Private/MuT/ASTOpReferenceResource.cpp
@@ -53,20 +53,21 @@ namespace UE::Mutable::Private
 
 
 	//-------------------------------------------------------------------------------------------------
-	void ASTOpReferenceResource::Link(FProgram& program, FLinkerOptions* Options)
+	void ASTOpReferenceResource::Link(FProgram& Program, FLinkerOptions* Options)
 	{
-		if (!linkedAddress)
+		if (!LinkedAddress)
 		{
-			OP::ResourceReferenceArgs Args;
+			FOperation::ResourceReferenceArgs Args;
 			FMemory::Memzero(Args);
 			Args.ID = ID;
 			Args.ForceLoad = bForceLoad ? 1 : 0;
 			Args.ImageDesc = ImageDesc;
 
-			linkedAddress = (OP::ADDRESS)program.OpAddress.Num();
-			program.OpAddress.Add((uint32)program.ByteCode.Num());
-			AppendCode(program.ByteCode, Type);
-			AppendCode(program.ByteCode, Args);
+			++Program.NumOps;
+			LinkedAddress = MakeProgramAddress(Type, Program.ByteCode.Num());
+			
+			AppendCode(Program.ByteCode, Type);
+			AppendCode(Program.ByteCode, Args);
 		}
 	}
 
@@ -139,4 +140,4 @@ namespace UE::Mutable::Private
 	}
 
 
-}
\ No newline at end of file
+}
diff --git a/Engine/Plugins/Mutable/Source/MutableTools/Private/MuT/ASTOpScalarArithmetic.cpp b/Engine/Plugins/Mutable/Source/MutableTools/Private/MuT/ASTOpScalarArithmetic.cpp
--- a/Engine/Plugins/Mutable/Source/MutableTools/Private/MuT/ASTOpScalarArithmetic.cpp
+++ b/Engine/Plugins/Mutable/Source/MutableTools/Private/MuT/ASTOpScalarArithmetic.cpp
@@ -69,17 +69,26 @@ namespace UE::Mutable::Private
 	void ASTOpScalarArithmetic::Link(FProgram& Program, FLinkerOptions*)
 	{
 		// Already linked?
-		if (!linkedAddress)
+		if (!LinkedAddress)
 		{
-			OP::ArithmeticArgs Args;
+			FOperation::ArithmeticArgs Args;
 			FMemory::Memzero(Args);
 
-			if (A) Args.A = A->linkedAddress;
-			if (B) Args.B = B->linkedAddress;
+			if (A) 
+			{
+				Args.A = A->LinkedAddress;
+			}
+
+			if (B) 
+			{
+				Args.B = B->LinkedAddress;
+			}
+
 			Args.Operation = Operation;
 
-			linkedAddress = (OP::ADDRESS)Program.OpAddress.Num();
-			Program.OpAddress.Add(Program.ByteCode.Num());
+			++Program.NumOps;
+			LinkedAddress = MakeProgramAddress(GetOpType(), Program.ByteCode.Num());
+			
 			AppendCode(Program.ByteCode, GetOpType());
 			AppendCode(Program.ByteCode, Args);
 		}
diff --git a/Engine/Plugins/Mutable/Source/MutableTools/Private/MuT/ASTOpScalarCurve.cpp b/Engine/Plugins/Mutable/Source/MutableTools/Private/MuT/ASTOpScalarCurve.cpp
--- a/Engine/Plugins/Mutable/Source/MutableTools/Private/MuT/ASTOpScalarCurve.cpp
+++ b/Engine/Plugins/Mutable/Source/MutableTools/Private/MuT/ASTOpScalarCurve.cpp
@@ -61,20 +61,21 @@ namespace UE::Mutable::Private
 	}
 
 
-	void ASTOpScalarCurve::Link(FProgram& program, FLinkerOptions*)
+	void ASTOpScalarCurve::Link(FProgram& Program, FLinkerOptions*)
 	{
-		if (!linkedAddress)
+		if (!LinkedAddress)
 		{
-			OP::ScalarCurveArgs Args;
+			FOperation::ScalarCurveArgs Args;
 			FMemory::Memzero(Args);
-			Args.time = time ? time->linkedAddress : 0;
-			Args.curve = program.AddConstant(Curve);
-
-			linkedAddress = (OP::ADDRESS)program.OpAddress.Num();
-			//program.m_code.push_back(op);
-			program.OpAddress.Add((uint32_t)program.ByteCode.Num());
-			AppendCode(program.ByteCode, EOpType::SC_CURVE);
-			AppendCode(program.ByteCode, Args);
+
+			Args.time = time ? time->LinkedAddress : 0;
+			Args.curve = Program.AddConstant(Curve);
+
+			++Program.NumOps;
+			LinkedAddress = MakeProgramAddress(GetOpType(), Program.ByteCode.Num());
+			
+			AppendCode(Program.ByteCode, GetOpType());
+			AppendCode(Program.ByteCode, Args);
 		}
 
 	}
diff --git a/Engine/Plugins/Mutable/Source/MutableTools/Private/MuT/ASTOpSwitch.cpp b/Engine/Plugins/Mutable/Source/MutableTools/Private/MuT/ASTOpSwitch.cpp
--- a/Engine/Plugins/Mutable/Source/MutableTools/Private/MuT/ASTOpSwitch.cpp
+++ b/Engine/Plugins/Mutable/Source/MutableTools/Private/MuT/ASTOpSwitch.cpp
@@ -181,15 +181,16 @@ namespace UE::Mutable::Private
 	void ASTOpSwitch::Link(FProgram& Program, FLinkerOptions*)
 	{
 		// Already linked?
-		if (!linkedAddress)
+		if (!LinkedAddress)
 		{
-			linkedAddress = (OP::ADDRESS)Program.OpAddress.Num();
-			Program.OpAddress.Add(Program.ByteCode.Num());
+			++Program.NumOps;
+			LinkedAddress = MakeProgramAddress(GetOpType(), Program.ByteCode.Num());
+			
+			AppendCode(Program.ByteCode, Type);
 
-			OP::ADDRESS VarAddress = Variable ? Variable->linkedAddress : 0;
-			OP::ADDRESS DefAddress = Default ? Default->linkedAddress : 0;
+			FOperation::ADDRESS VarAddress = Variable ? Variable->LinkedAddress : 0;
+			FOperation::ADDRESS DefAddress = Default ? Default->LinkedAddress : 0;
 
-			AppendCode(Program.ByteCode, Type);
 			AppendCode(Program.ByteCode, VarAddress);
 			AppendCode(Program.ByteCode, DefAddress);
 
@@ -218,21 +219,21 @@ namespace UE::Mutable::Private
 
 			if (!bUseRanges)
 			{
-				OP::FSwitchCaseDescriptor CaseDesc;
+				FOperation::FSwitchCaseDescriptor CaseDesc;
 				CaseDesc.Count = Cases.Num();
 				CaseDesc.bUseRanges = false;
 				AppendCode(Program.ByteCode, CaseDesc);
 
 				for (const FCase& Case : Cases)
 				{
-					OP::ADDRESS CaseBranchAddress = Case.Branch ? Case.Branch->linkedAddress : 0;
+					FOperation::ADDRESS CaseBranchAddress = Case.Branch ? Case.Branch->LinkedAddress : 0;
 					AppendCode(Program.ByteCode, Case.Condition);
 					AppendCode(Program.ByteCode, CaseBranchAddress);
 				}
 			}
 			else
 			{
-				OP::FSwitchCaseDescriptor CaseDesc;
+				FOperation::FSwitchCaseDescriptor CaseDesc;
 				CaseDesc.Count = Cases.Num() - CasesInRange;
 				CaseDesc.bUseRanges = true;
 
@@ -258,7 +259,7 @@ namespace UE::Mutable::Private
 					}
 
 					const FCase& FirstCase = Cases[RangeStart];
-					OP::ADDRESS CaseBranchAddress = FirstCase.Branch ? FirstCase.Branch->linkedAddress : 0;
+					FOperation::ADDRESS CaseBranchAddress = FirstCase.Branch ? FirstCase.Branch->LinkedAddress : 0;
 					AppendCode(Program.ByteCode, FirstCase.Condition);
 					AppendCode(Program.ByteCode, RangeSize);
 					AppendCode(Program.ByteCode, CaseBranchAddress);
diff --git a/Engine/Plugins/Mutable/Source/MutableTools/Private/MuT/CodeGenerator.cpp b/Engine/Plugins/Mutable/Source/MutableTools/Private/MuT/CodeGenerator.cpp
--- a/Engine/Plugins/Mutable/Source/MutableTools/Private/MuT/CodeGenerator.cpp
+++ b/Engine/Plugins/Mutable/Source/MutableTools/Private/MuT/CodeGenerator.cpp
@@ -434,7 +434,7 @@ namespace UE::Mutable::Private
 
 			// When we patch from edit nodes, we want to apply it to all the channels.
 			// \todo: since we can choose the patch function, maybe we want to be able to select this as well.
-			LayerOp->Flags = Patch.bPatchApplyToAlpha ? OP::ImageLayerArgs::F_APPLY_TO_ALPHA : 0;
+			LayerOp->Flags = Patch.bPatchApplyToAlpha ? FOperation::ImageLayerArgs::F_APPLY_TO_ALPHA : 0;
 
 			NodeImage* ImageNode = Patch.PatchImage.get();
 			Ptr<ASTOp> BlendOp;
diff --git a/Engine/Plugins/Mutable/Source/MutableTools/Private/MuT/CodeGenerator_Color.cpp b/Engine/Plugins/Mutable/Source/MutableTools/Private/MuT/CodeGenerator_Color.cpp
--- a/Engine/Plugins/Mutable/Source/MutableTools/Private/MuT/CodeGenerator_Color.cpp
+++ b/Engine/Plugins/Mutable/Source/MutableTools/Private/MuT/CodeGenerator_Color.cpp
@@ -466,7 +466,6 @@ namespace UE::Mutable::Private
 	}
 
 
-	//-------------------------------------------------------------------------------------------------
 	void CodeGenerator::GenerateColor_Arithmetic(FColorGenerationResult& result, const FGenericGenerationOptions& Options, const Ptr<const NodeColourArithmeticOperation>& Typed)
 	{
 		const NodeColourArithmeticOperation& node = *Typed;
@@ -475,13 +474,13 @@ namespace UE::Mutable::Private
 
 		switch (node.Operation)
 		{
-		case NodeColourArithmeticOperation::EOperation::Add: op->Operation = OP::ArithmeticArgs::ADD; break;
-		case NodeColourArithmeticOperation::EOperation::Subtract: op->Operation = OP::ArithmeticArgs::SUBTRACT; break;
-		case NodeColourArithmeticOperation::EOperation::Multiply: op->Operation = OP::ArithmeticArgs::MULTIPLY; break;
-		case NodeColourArithmeticOperation::EOperation::Divide: op->Operation = OP::ArithmeticArgs::DIVIDE; break;
+		case NodeColourArithmeticOperation::EOperation::Add: op->Operation = FOperation::ArithmeticArgs::ADD; break;
+		case NodeColourArithmeticOperation::EOperation::Subtract: op->Operation = FOperation::ArithmeticArgs::SUBTRACT; break;
+		case NodeColourArithmeticOperation::EOperation::Multiply: op->Operation = FOperation::ArithmeticArgs::MULTIPLY; break;
+		case NodeColourArithmeticOperation::EOperation::Divide: op->Operation = FOperation::ArithmeticArgs::DIVIDE; break;
 		default:
 			checkf(false, TEXT("Unknown arithmetic operation."));
-			op->Operation = OP::ArithmeticArgs::NONE;
+			op->Operation = FOperation::ArithmeticArgs::NONE;
 			break;
 		}
 
diff --git a/Engine/Plugins/Mutable/Source/MutableTools/Private/MuT/CodeGenerator_Image.cpp b/Engine/Plugins/Mutable/Source/MutableTools/Private/MuT/CodeGenerator_Image.cpp
--- a/Engine/Plugins/Mutable/Source/MutableTools/Private/MuT/CodeGenerator_Image.cpp
+++ b/Engine/Plugins/Mutable/Source/MutableTools/Private/MuT/CodeGenerator_Image.cpp
@@ -1480,9 +1480,11 @@ namespace UE::Mutable::Private
                 cop->SetValue(FormatMeshResult);
 
                 Ptr<ASTOpMeshFormat> FormatOp = new ASTOpMeshFormat();
-				FormatOp->Flags = OP::MeshFormatArgs::Vertex
-                        | OP::MeshFormatArgs::Index
-                        | OP::MeshFormatArgs::ResetBufferIndices;
+				FormatOp->Flags = 
+					FOperation::MeshFormatArgs::Vertex | 
+					FOperation::MeshFormatArgs::Index  | 
+					FOperation::MeshFormatArgs::ResetBufferIndices;
+
 				FormatOp->Format = cop;
 				FormatOp->Source = CurrentMeshToProjectOp;
 				CurrentMeshToProjectOp = FormatOp;
@@ -1490,7 +1492,7 @@ namespace UE::Mutable::Private
             else
             {
                 // Extract the mesh layout block
-                if ( GeneratedLayoutBlockId!= FLayoutBlock::InvalidBlockId)
+                if (GeneratedLayoutBlockId!= FLayoutBlock::InvalidBlockId)
                 {
                     Ptr<ASTOpMeshExtractLayoutBlocks> eop = new ASTOpMeshExtractLayoutBlocks();
                     eop->Source = CurrentMeshToProjectOp;
@@ -1512,9 +1514,9 @@ namespace UE::Mutable::Private
                 cop->SetValue(FormatMeshResult);
 
                 Ptr<ASTOpMeshFormat> FormatOp = new ASTOpMeshFormat();
-				FormatOp->Flags = OP::MeshFormatArgs::Vertex
-					| OP::MeshFormatArgs::Index
-					| OP::MeshFormatArgs::ResetBufferIndices;
+				FormatOp->Flags = FOperation::MeshFormatArgs::Vertex
+					| FOperation::MeshFormatArgs::Index
+					| FOperation::MeshFormatArgs::ResetBufferIndices;
 				FormatOp->Format = cop;
 				FormatOp->Source = CurrentMeshToProjectOp;
 				CurrentMeshToProjectOp = FormatOp;
diff --git a/Engine/Plugins/Mutable/Source/MutableTools/Private/MuT/CodeGenerator_Mesh.cpp b/Engine/Plugins/Mutable/Source/MutableTools/Private/MuT/CodeGenerator_Mesh.cpp
--- a/Engine/Plugins/Mutable/Source/MutableTools/Private/MuT/CodeGenerator_Mesh.cpp
+++ b/Engine/Plugins/Mutable/Source/MutableTools/Private/MuT/CodeGenerator_Mesh.cpp
@@ -1004,19 +1004,19 @@ namespace UE::Mutable::Private
 
 				if (Node->VertexBuffers.GetBufferCount())
 				{
-					op->Flags |= OP::MeshFormatArgs::Vertex;
+					op->Flags |= FOperation::MeshFormatArgs::Vertex;
 					FormatMesh->VertexBuffers = Node->VertexBuffers;
 				}
 
 				if (Node->IndexBuffers.GetBufferCount())
 				{
-					op->Flags |= OP::MeshFormatArgs::Index;
+					op->Flags |= FOperation::MeshFormatArgs::Index;
 					FormatMesh->IndexBuffers = Node->IndexBuffers;
 				}
 
 				if (Node->bOptimizeBuffers)
 				{
-					op->Flags |= OP::MeshFormatArgs::OptimizeBuffers;
+					op->Flags |= FOperation::MeshFormatArgs::OptimizeBuffers;
 				}
 
 				Ptr<ASTOpConstantResource> cop = new ASTOpConstantResource();
diff --git a/Engine/Plugins/Mutable/Source/MutableTools/Private/MuT/CodeGenerator_Scalar.cpp b/Engine/Plugins/Mutable/Source/MutableTools/Private/MuT/CodeGenerator_Scalar.cpp
--- a/Engine/Plugins/Mutable/Source/MutableTools/Private/MuT/CodeGenerator_Scalar.cpp
+++ b/Engine/Plugins/Mutable/Source/MutableTools/Private/MuT/CodeGenerator_Scalar.cpp
@@ -377,13 +377,13 @@ namespace UE::Mutable::Private
 
 		switch (Node->Operation)
 		{
-		case NodeScalarArithmeticOperation::AO_ADD: op->Operation = OP::ArithmeticArgs::ADD; break;
-		case NodeScalarArithmeticOperation::AO_SUBTRACT: op->Operation = OP::ArithmeticArgs::SUBTRACT; break;
-		case NodeScalarArithmeticOperation::AO_MULTIPLY: op->Operation = OP::ArithmeticArgs::MULTIPLY; break;
-		case NodeScalarArithmeticOperation::AO_DIVIDE: op->Operation = OP::ArithmeticArgs::DIVIDE; break;
+		case NodeScalarArithmeticOperation::AO_ADD: op->Operation = FOperation::ArithmeticArgs::ADD; break;
+		case NodeScalarArithmeticOperation::AO_SUBTRACT: op->Operation = FOperation::ArithmeticArgs::SUBTRACT; break;
+		case NodeScalarArithmeticOperation::AO_MULTIPLY: op->Operation = FOperation::ArithmeticArgs::MULTIPLY; break;
+		case NodeScalarArithmeticOperation::AO_DIVIDE: op->Operation = FOperation::ArithmeticArgs::DIVIDE; break;
 		default:
 			checkf(false, TEXT("Unknown arithmetic operation."));
-			op->Operation = OP::ArithmeticArgs::NONE;
+			op->Operation = FOperation::ArithmeticArgs::NONE;
 			break;
 		}
 
@@ -483,4 +483,4 @@ namespace UE::Mutable::Private
 		return ParamResult.op;
 	}
 
-}
\ No newline at end of file
+}
diff --git a/Engine/Plugins/Mutable/Source/MutableTools/Private/MuT/CodeOptimiser.cpp b/Engine/Plugins/Mutable/Source/MutableTools/Private/MuT/CodeOptimiser.cpp
--- a/Engine/Plugins/Mutable/Source/MutableTools/Private/MuT/CodeOptimiser.cpp
+++ b/Engine/Plugins/Mutable/Source/MutableTools/Private/MuT/CodeOptimiser.cpp
@@ -187,7 +187,7 @@ namespace UE::Mutable::Private
 			LinkerOptions.bSeparateImageMips = false;
 
 			TSharedPtr<FModel> Model = MakeShared<FModel>();
-			OP::ADDRESS at = ASTOp::FullLink(SourceCloned, Model->GetProgram(), &LinkerOptions);
+			FOperation::ADDRESS at = ASTOp::FullLink(SourceCloned, Model->GetProgram(), &LinkerOptions);
 
 			FProgram::FState state;
 			state.Root = at;
diff --git a/Engine/Plugins/Mutable/Source/MutableTools/Private/MuT/Compiler.cpp b/Engine/Plugins/Mutable/Source/MutableTools/Private/MuT/Compiler.cpp
--- a/Engine/Plugins/Mutable/Source/MutableTools/Private/MuT/Compiler.cpp
+++ b/Engine/Plugins/Mutable/Source/MutableTools/Private/MuT/Compiler.cpp
@@ -55,68 +55,31 @@
 		{ false,	{ 0, 0, 0, 0, 0, 0, 0, 0, 0, 0, 0, 0, 0, 0, 0, 0, 0, 0, 0, 0, 0, 0, 0, 0, 0, 0, 0, 0, 0 }		},	// NONE
 
 		{ false,	{ 0, 0, 0, 0, 0, 0, 0, 0, 0, 0, 0, 0, 0, 0, 0, 0, 0, 0, 0, 0, 0, 0, 0, 0, 0, 0, 0, 0, 0 }		},	// BO_CONSTANT
-		{ false,	{ 0, 0, 0, 0, 0, 0, 0, 0, 0, 0, 0, 0, 0, 0, 0, 0, 0, 0, 0, 0, 0, 0, 0, 0, 0, 0, 0, 0, 0 }		},	// NU_CONSTANT
-		{ false,	{ 0, 0, 0, 0, 0, 0, 0, 0, 0, 0, 0, 0, 0, 0, 0, 0, 0, 0, 0, 0, 0, 0, 0, 0, 0, 0, 0, 0, 0 }		},	// SC_CONSTANT
-		{ false,	{ 0, 0, 0, 0, 0, 0, 0, 0, 0, 0, 0, 0, 0, 0, 0, 0, 0, 0, 0, 0, 0, 0, 0, 0, 0, 0, 0, 0, 0 }		},	// CO_CONSTANT
-		{ true,		{ 1, 1, 1, 1, 1, 1, 1, 1, 1, 1, 1, 1, 0, 0, 0, 0, 0, 0, 0, 0, 0, 0, 0, 0, 0, 0, 0, 0, 0 }		},	// IM_CONSTANT
-		{ true,		{ 0, 0, 0, 0, 0, 0, 0, 0, 0, 0, 0, 0, 0, 0, 0, 0, 0, 0, 0, 0, 0, 0, 0, 0, 0, 0, 0, 0, 0 }		},	// ME_CONSTANT
-		{ false,	{ 0, 0, 0, 0, 0, 0, 0, 0, 0, 0, 0, 0, 0, 0, 0, 0, 0, 0, 0, 0, 0, 0, 0, 0, 0, 0, 0, 0, 0 }		},	// LA_CONSTANT
-		{ false,	{ 0, 0, 0, 0, 0, 0, 0, 0, 0, 0, 0, 0, 0, 0, 0, 0, 0, 0, 0, 0, 0, 0, 0, 0, 0, 0, 0, 0, 0 }		},	// PR_CONSTANT
-		{ false,	{ 0, 0, 0, 0, 0, 0, 0, 0, 0, 0, 0, 0, 0, 0, 0, 0, 0, 0, 0, 0, 0, 0, 0, 0, 0, 0, 0, 0, 0 }		},	// ST_CONSTANT
-		{ true,		{ 0, 0, 0, 0, 0, 0, 0, 0, 0, 0, 0, 0, 0, 0, 0, 0, 0, 0, 0, 0, 0, 0, 0, 0, 0, 0, 0, 0, 0 }		},	// ED_CONSTANT
-		{ false,	{ 0, 0, 0, 0, 0, 0, 0, 0, 0, 0, 0, 0, 0, 0, 0, 0, 0, 0, 0, 0, 0, 0, 0, 0, 0, 0, 0, 0, 0 }		},	// MA_CONSTANT
-		{ false,	{ 0, 0, 0, 0, 0, 0, 0, 0, 0, 0, 0, 0, 0, 0, 0, 0, 0, 0, 0, 0, 0, 0, 0, 0, 0, 0, 0, 0, 0 }		},	// MI_CONSTANT
-
 		{ false,	{ 0, 0, 0, 0, 0, 0, 0, 0, 0, 0, 0, 0, 0, 0, 0, 0, 0, 0, 0, 0, 0, 0, 0, 0, 0, 0, 0, 0, 0 }		},	// BO_PARAMETER
+		{ false,	{ 0, 0, 0, 0, 0, 0, 0, 0, 0, 0, 0, 0, 0, 0, 0, 0, 0, 0, 0, 0, 0, 0, 0, 0, 0, 0, 0, 0, 0 }		},	// BO_EQUAL_SC_CONST
+		{ false,	{ 0, 0, 0, 0, 0, 0, 0, 0, 0, 0, 0, 0, 0, 0, 0, 0, 0, 0, 0, 0, 0, 0, 0, 0, 0, 0, 0, 0, 0 }		},	// BO_AND
+		{ false,	{ 0, 0, 0, 0, 0, 0, 0, 0, 0, 0, 0, 0, 0, 0, 0, 0, 0, 0, 0, 0, 0, 0, 0, 0, 0, 0, 0, 0, 0 }		},	// BO_OR
+		{ false,	{ 0, 0, 0, 0, 0, 0, 0, 0, 0, 0, 0, 0, 0, 0, 0, 0, 0, 0, 0, 0, 0, 0, 0, 0, 0, 0, 0, 0, 0 }		},	// BO_NOT
+
+		{ false,	{ 0, 0, 0, 0, 0, 0, 0, 0, 0, 0, 0, 0, 0, 0, 0, 0, 0, 0, 0, 0, 0, 0, 0, 0, 0, 0, 0, 0, 0 }		},	// NU_CONSTANT
 		{ false,	{ 0, 0, 0, 0, 0, 0, 0, 0, 0, 0, 0, 0, 0, 0, 0, 0, 0, 0, 0, 0, 0, 0, 0, 0, 0, 0, 0, 0, 0 }		},	// NU_PARAMETER
-		{ false,	{ 0, 0, 0, 0, 0, 0, 0, 0, 0, 0, 0, 0, 0, 0, 0, 0, 0, 0, 0, 0, 0, 0, 0, 0, 0, 0, 0, 0, 0 }		},	// SC_PARAMETER
-		{ false,	{ 0, 0, 0, 0, 0, 0, 0, 0, 0, 0, 0, 0, 0, 0, 0, 0, 0, 0, 0, 0, 0, 0, 0, 0, 0, 0, 0, 0, 0 }		},	// CO_PARAMETER
-		{ false,	{ 0, 0, 0, 0, 0, 0, 0, 0, 0, 0, 0, 0, 0, 0, 0, 0, 0, 0, 0, 0, 0, 0, 0, 0, 0, 0, 0, 0, 0 }		},	// PR_PARAMETER
-		{ false,	{ 0, 0, 0, 0, 0, 0, 0, 0, 0, 0, 0, 0, 0, 0, 0, 0, 0, 0, 0, 0, 0, 0, 0, 0, 0, 0, 0, 0, 0 }		},	// IM_PARAMETER
-		{ false,	{ 0, 0, 0, 0, 0, 0, 0, 0, 0, 0, 0, 0, 0, 0, 0, 0, 0, 0, 0, 0, 0, 0, 0, 0, 0, 0, 0, 0, 0 }		},	// SK_PARAMETER
-		{ false,	{ 0, 0, 0, 0, 0, 0, 0, 0, 0, 0, 0, 0, 0, 0, 0, 0, 0, 0, 0, 0, 0, 0, 0, 0, 0, 0, 0, 0, 0 }		},	// ST_PARAMETER
-		{ false,	{ 0, 0, 0, 0, 0, 0, 0, 0, 0, 0, 0, 0, 0, 0, 0, 0, 0, 0, 0, 0, 0, 0, 0, 0, 0, 0, 0, 0, 0 }		},	// MA_PARAMETER
-		{ false,	{ 0, 0, 0, 0, 0, 0, 0, 0, 0, 0, 0, 0, 0, 0, 0, 0, 0, 0, 0, 0, 0, 0, 0, 0, 0, 0, 0, 0, 0 }		},	// MI_PARAMETER
-		{ false,	{ 0, 0, 0, 0, 0, 0, 0, 0, 0, 0, 0, 0, 0, 0, 0, 0, 0, 0, 0, 0, 0, 0, 0, 0, 0, 0, 0, 0, 0 }		},	// IS_PARAMETER
-
-		{ true,		{ 0, 0, 0, 0, 0, 0, 0, 0, 0, 0, 0, 0, 0, 0, 0, 0, 0, 0, 0, 0, 0, 0, 0, 0, 0, 0, 0, 0, 0 }		},	// IM_REFERENCE
-		{ true,		{ 0, 0, 0, 0, 0, 0, 0, 0, 0, 0, 0, 0, 0, 0, 0, 0, 0, 0, 0, 0, 0, 0, 0, 0, 0, 0, 0, 0, 0 }		},	// ME_REFERENCE
-
 		{ false,	{ 0, 0, 0, 0, 0, 0, 0, 0, 0, 0, 0, 0, 0, 0, 0, 0, 0, 0, 0, 0, 0, 0, 0, 0, 0, 0, 0, 0, 0 }		},	// NU_CONDITIONAL
-		{ false,	{ 0, 0, 0, 0, 0, 0, 0, 0, 0, 0, 0, 0, 0, 0, 0, 0, 0, 0, 0, 0, 0, 0, 0, 0, 0, 0, 0, 0, 0 }		},	// SC_CONDITIONAL
-		{ false,	{ 0, 0, 0, 0, 0, 0, 0, 0, 0, 0, 0, 0, 0, 0, 0, 0, 0, 0, 0, 0, 0, 0, 0, 0, 0, 0, 0, 0, 0 }		},	// CO_CONDITIONAL
-		{ true,		{ 1, 1, 1, 1, 1, 1, 1, 1, 1, 1, 1, 1, 1, 1, 1, 1, 1, 1, 1, 1, 1, 1, 0, 0, 1, 1, 0, 0, 0 }		},	// IM_CONDITIONAL
-		{ true,		{ 0, 0, 0, 0, 0, 0, 0, 0, 0, 0, 0, 0, 0, 0, 0, 0, 0, 0, 0, 0, 0, 0, 0, 0, 0, 0, 0, 0, 0 }		},	// ME_CONDITIONAL
-		{ false,	{ 0, 0, 0, 0, 0, 0, 0, 0, 0, 0, 0, 0, 0, 0, 0, 0, 0, 0, 0, 0, 0, 0, 0, 0, 0, 0, 0, 0, 0 }		},	// LA_CONDITIONAL
-		{ false,	{ 0, 0, 0, 0, 0, 0, 0, 0, 0, 0, 0, 0, 0, 0, 0, 0, 0, 0, 0, 0, 0, 0, 0, 0, 0, 0, 0, 0, 0 }		},	// IN_CONDITIONAL
-		{ false,	{ 0, 0, 0, 0, 0, 0, 0, 0, 0, 0, 0, 0, 0, 0, 0, 0, 0, 0, 0, 0, 0, 0, 0, 0, 0, 0, 0, 0, 0 }		},	// ED_CONDITIONAL
-		{ false,	{ 0, 0, 0, 0, 0, 0, 0, 0, 0, 0, 0, 0, 0, 0, 0, 0, 0, 0, 0, 0, 0, 0, 0, 0, 0, 0, 0, 0, 0 }		},	// MI_CONDITIONAL
-
 		{ false,	{ 0, 0, 0, 0, 0, 0, 0, 0, 0, 0, 0, 0, 0, 0, 0, 0, 0, 0, 0, 0, 0, 0, 0, 0, 0, 0, 0, 0, 0 }		},	// NU_SWITCH
+
+		{ false,	{ 0, 0, 0, 0, 0, 0, 0, 0, 0, 0, 0, 0, 0, 0, 0, 0, 0, 0, 0, 0, 0, 0, 0, 0, 0, 0, 0, 0, 0 }		},	// SC_CONSTANT
+		{ false,	{ 0, 0, 0, 0, 0, 0, 0, 0, 0, 0, 0, 0, 0, 0, 0, 0, 0, 0, 0, 0, 0, 0, 0, 0, 0, 0, 0, 0, 0 }		},	// SC_PARAMETER
+		{ false,	{ 0, 0, 0, 0, 0, 0, 0, 0, 0, 0, 0, 0, 0, 0, 0, 0, 0, 0, 0, 0, 0, 0, 0, 0, 0, 0, 0, 0, 0 }		},	// SC_CONDITIONAL
 		{ false,	{ 0, 0, 0, 0, 0, 0, 0, 0, 0, 0, 0, 0, 0, 0, 0, 0, 0, 0, 0, 0, 0, 0, 0, 0, 0, 0, 0, 0, 0 }		},	// SC_SWITCH
-		{ false,	{ 0, 0, 0, 0, 0, 0, 0, 0, 0, 0, 0, 0, 0, 0, 0, 0, 0, 0, 0, 0, 0, 0, 0, 0, 0, 0, 0, 0, 0 }		},	// CO_SWITCH
-		{ true,		{ 1, 1, 1, 1, 1, 1, 1, 1, 1, 1, 1, 1, 1, 1, 1, 1, 1, 1, 1, 1, 1, 1, 0, 0, 1, 1, 0, 0, 0 }		},	// IM_SWITCH
-		{ true,		{ 0, 0, 0, 0, 0, 0, 0, 0, 0, 0, 0, 0, 0, 0, 0, 0, 0, 0, 0, 0, 0, 0, 0, 0, 0, 0, 0, 0, 0 }		},	// ME_SWITCH
-		{ false,	{ 0, 0, 0, 0, 0, 0, 0, 0, 0, 0, 0, 0, 0, 0, 0, 0, 0, 0, 0, 0, 0, 0, 0, 0, 0, 0, 0, 0, 0 }		},	// LA_SWITCH
-		{ false,	{ 0, 0, 0, 0, 0, 0, 0, 0, 0, 0, 0, 0, 0, 0, 0, 0, 0, 0, 0, 0, 0, 0, 0, 0, 0, 0, 0, 0, 0 }		},	// IN_SWITCH
-		{ false,	{ 0, 0, 0, 0, 0, 0, 0, 0, 0, 0, 0, 0, 0, 0, 0, 0, 0, 0, 0, 0, 0, 0, 0, 0, 0, 0, 0, 0, 0 }		},	// ED_SWITCH
-		{ false,	{ 0, 0, 0, 0, 0, 0, 0, 0, 0, 0, 0, 0, 0, 0, 0, 0, 0, 0, 0, 0, 0, 0, 0, 0, 0, 0, 0, 0, 0 }		},	// MI_SWITCH
-
 		{ false,	{ 0, 0, 0, 0, 0, 0, 0, 0, 0, 0, 0, 0, 0, 0, 0, 0, 0, 0, 0, 0, 0, 0, 0, 0, 0, 0, 0, 0, 0 }		},	// SC_MATERIAL_BREAK
-		{ false,	{ 0, 0, 0, 0, 0, 0, 0, 0, 0, 0, 0, 0, 0, 0, 0, 0, 0, 0, 0, 0, 0, 0, 0, 0, 0, 0, 0, 0, 0 }		},	// CO_MATERIAL_BREAK
-		{ false,	{ 0, 0, 0, 0, 0, 0, 0, 0, 0, 0, 0, 0, 0, 0, 0, 0, 0, 0, 0, 0, 0, 0, 0, 0, 0, 0, 0, 0, 0 }		},	// IM_MATERIAL_BREAK
-		{ false,	{ 0, 0, 0, 0, 0, 0, 0, 0, 0, 0, 0, 0, 0, 0, 0, 0, 0, 0, 0, 0, 0, 0, 0, 0, 0, 0, 0, 0, 0 }		},	// IM_PARAMETER_FROM_MATERIAL
-		{ false,	{ 0, 0, 0, 0, 0, 0, 0, 0, 0, 0, 0, 0, 0, 0, 0, 0, 0, 0, 0, 0, 0, 0, 0, 0, 0, 0, 0, 0, 0 }		},	// MI_MODIFY
-
-		{ false,	{ 0, 0, 0, 0, 0, 0, 0, 0, 0, 0, 0, 0, 0, 0, 0, 0, 0, 0, 0, 0, 0, 0, 0, 0, 0, 0, 0, 0, 0 }		},	// BO_EQUAL_SC_CONST
-		{ false,	{ 0, 0, 0, 0, 0, 0, 0, 0, 0, 0, 0, 0, 0, 0, 0, 0, 0, 0, 0, 0, 0, 0, 0, 0, 0, 0, 0, 0, 0 }		},	// BO_AND
-		{ false,	{ 0, 0, 0, 0, 0, 0, 0, 0, 0, 0, 0, 0, 0, 0, 0, 0, 0, 0, 0, 0, 0, 0, 0, 0, 0, 0, 0, 0, 0 }		},	// BO_OR
-		{ false,	{ 0, 0, 0, 0, 0, 0, 0, 0, 0, 0, 0, 0, 0, 0, 0, 0, 0, 0, 0, 0, 0, 0, 0, 0, 0, 0, 0, 0, 0 }		},	// BO_NOT
-
 		{ true,		{ 0, 0, 0, 0, 0, 0, 0, 0, 0, 0, 0, 0, 0, 0, 0, 0, 0, 0, 0, 0, 0, 0, 0, 0, 0, 0, 0, 0, 0 }		},	// SC_ARITHMETIC
 		{ true,		{ 0, 0, 0, 0, 0, 0, 0, 0, 0, 0, 0, 0, 0, 0, 0, 0, 0, 0, 0, 0, 0, 0, 0, 0, 0, 0, 0, 0, 0 }		},	// SC_CURVE
 		{ true,		{ 0, 0, 0, 0, 0, 0, 0, 0, 0, 0, 0, 0, 0, 0, 0, 0, 0, 0, 0, 0, 0, 0, 0, 0, 0, 0, 0, 0, 0 }		},	// SC_EXTERNAL
 
+		{ false,	{ 0, 0, 0, 0, 0, 0, 0, 0, 0, 0, 0, 0, 0, 0, 0, 0, 0, 0, 0, 0, 0, 0, 0, 0, 0, 0, 0, 0, 0 }		},	// CO_CONSTANT
+		{ false,	{ 0, 0, 0, 0, 0, 0, 0, 0, 0, 0, 0, 0, 0, 0, 0, 0, 0, 0, 0, 0, 0, 0, 0, 0, 0, 0, 0, 0, 0 }		},	// CO_PARAMETER
+		{ false,	{ 0, 0, 0, 0, 0, 0, 0, 0, 0, 0, 0, 0, 0, 0, 0, 0, 0, 0, 0, 0, 0, 0, 0, 0, 0, 0, 0, 0, 0 }		},	// CO_CONDITIONAL
+		{ false,	{ 0, 0, 0, 0, 0, 0, 0, 0, 0, 0, 0, 0, 0, 0, 0, 0, 0, 0, 0, 0, 0, 0, 0, 0, 0, 0, 0, 0, 0 }		},	// CO_SWITCH
+		{ false,	{ 0, 0, 0, 0, 0, 0, 0, 0, 0, 0, 0, 0, 0, 0, 0, 0, 0, 0, 0, 0, 0, 0, 0, 0, 0, 0, 0, 0, 0 }		},	// CO_MATERIAL_BREAK
 		{ false,	{ 0, 0, 0, 0, 0, 0, 0, 0, 0, 0, 0, 0, 0, 0, 0, 0, 0, 0, 0, 0, 0, 0, 0, 0, 0, 0, 0, 0, 0 }		},	// CO_SAMPLEIMAGE
 		{ false,	{ 0, 0, 0, 0, 0, 0, 0, 0, 0, 0, 0, 0, 0, 0, 0, 0, 0, 0, 0, 0, 0, 0, 0, 0, 0, 0, 0, 0, 0 }		},	// CO_SWIZZLE
 		{ false,	{ 0, 0, 0, 0, 0, 0, 0, 0, 0, 0, 0, 0, 0, 0, 0, 0, 0, 0, 0, 0, 0, 0, 0, 0, 0, 0, 0, 0, 0 }		},	// CO_FROMSCALARS
@@ -124,6 +87,19 @@
 		{ false,	{ 0, 0, 0, 0, 0, 0, 0, 0, 0, 0, 0, 0, 0, 0, 0, 0, 0, 0, 0, 0, 0, 0, 0, 0, 0, 0, 0, 0, 0 }		},	// CO_LINEARTOSRGB
 		{ true,		{ 0, 0, 0, 0, 0, 0, 0, 0, 0, 0, 0, 0, 0, 0, 0, 0, 0, 0, 0, 0, 0, 0, 0, 0, 0, 0, 0, 0, 0 }		},	// CO_EXTERNAL
 
+		{ false,	{ 0, 0, 0, 0, 0, 0, 0, 0, 0, 0, 0, 0, 0, 0, 0, 0, 0, 0, 0, 0, 0, 0, 0, 0, 0, 0, 0, 0, 0 }		},	// ST_CONSTANT
+		{ false,	{ 0, 0, 0, 0, 0, 0, 0, 0, 0, 0, 0, 0, 0, 0, 0, 0, 0, 0, 0, 0, 0, 0, 0, 0, 0, 0, 0, 0, 0 }		},	// ST_PARAMETER
+
+		{ false,	{ 0, 0, 0, 0, 0, 0, 0, 0, 0, 0, 0, 0, 0, 0, 0, 0, 0, 0, 0, 0, 0, 0, 0, 0, 0, 0, 0, 0, 0 }		},	// PR_CONSTANT
+		{ false,	{ 0, 0, 0, 0, 0, 0, 0, 0, 0, 0, 0, 0, 0, 0, 0, 0, 0, 0, 0, 0, 0, 0, 0, 0, 0, 0, 0, 0, 0 }		},	// PR_PARAMETER
+
+		{ true,		{ 1, 1, 1, 1, 1, 1, 1, 1, 1, 1, 1, 1, 0, 0, 0, 0, 0, 0, 0, 0, 0, 0, 0, 0, 0, 0, 0, 0, 0 }		},	// IM_CONSTANT
+		{ false,	{ 0, 0, 0, 0, 0, 0, 0, 0, 0, 0, 0, 0, 0, 0, 0, 0, 0, 0, 0, 0, 0, 0, 0, 0, 0, 0, 0, 0, 0 }		},	// IM_PARAMETER
+		{ true,		{ 0, 0, 0, 0, 0, 0, 0, 0, 0, 0, 0, 0, 0, 0, 0, 0, 0, 0, 0, 0, 0, 0, 0, 0, 0, 0, 0, 0, 0 }		},	// IM_REFERENCE
+		{ true,		{ 1, 1, 1, 1, 1, 1, 1, 1, 1, 1, 1, 1, 1, 1, 1, 1, 1, 1, 1, 1, 1, 1, 0, 0, 1, 1, 0, 0, 0 }		},	// IM_CONDITIONAL
+		{ true,		{ 1, 1, 1, 1, 1, 1, 1, 1, 1, 1, 1, 1, 1, 1, 1, 1, 1, 1, 1, 1, 1, 1, 0, 0, 1, 1, 0, 0, 0 }		},	// IM_SWITCH
+		{ false,	{ 0, 0, 0, 0, 0, 0, 0, 0, 0, 0, 0, 0, 0, 0, 0, 0, 0, 0, 0, 0, 0, 0, 0, 0, 0, 0, 0, 0, 0 }		},	// IM_MATERIAL_BREAK
+		{ false,	{ 0, 0, 0, 0, 0, 0, 0, 0, 0, 0, 0, 0, 0, 0, 0, 0, 0, 0, 0, 0, 0, 0, 0, 0, 0, 0, 0, 0, 0 }		},	// IM_PARAMETER_FROM_MATERIAL
 		{ true,		{ 0, 1, 1, 1, 0, 0, 0, 0, 0, 0, 0, 0, 0, 0, 0, 0, 0, 0, 0, 0, 0, 0, 0, 0, 0, 0, 0, 0, 0 }		},	// IM_LAYER
 		{ true,		{ 0, 1, 1, 1, 0, 0, 0, 0, 0, 0, 0, 0, 0, 0, 0, 0, 0, 0, 0, 0, 0, 0, 0, 0, 0, 0, 0, 0, 0 }		},	// IM_LAYERCOLOUR
 		{ true,		{ 1, 1, 1, 1, 1, 1, 1, 1, 1, 1, 1, 1, 0, 0, 0, 0, 0, 0, 0, 0, 0, 0, 0, 0, 0, 0, 0, 0, 0 }		},	// IM_PIXELFORMAT
@@ -153,6 +129,11 @@
 		{ true,		{ 0, 0, 0, 0, 0, 0, 0, 0, 0, 0, 0, 0, 0, 0, 0, 0, 0, 0, 0, 0, 0, 0, 0, 0, 0, 0, 0, 0, 0 }		},	// IM_PARAMETER_CONVERT
 		{ true,		{ 0, 0, 0, 0, 0, 0, 0, 0, 0, 0, 0, 0, 0, 0, 0, 0, 0, 0, 0, 0, 0, 0, 0, 0, 0, 0, 0, 0, 0 }		},	// IM_EXTERNAL
 
+		{ true,		{ 0, 0, 0, 0, 0, 0, 0, 0, 0, 0, 0, 0, 0, 0, 0, 0, 0, 0, 0, 0, 0, 0, 0, 0, 0, 0, 0, 0, 0 }		},	// ME_CONSTANT
+		{ false,	{ 0, 0, 0, 0, 0, 0, 0, 0, 0, 0, 0, 0, 0, 0, 0, 0, 0, 0, 0, 0, 0, 0, 0, 0, 0, 0, 0, 0, 0 }		},	// ME_PARAMETER
+		{ true,		{ 0, 0, 0, 0, 0, 0, 0, 0, 0, 0, 0, 0, 0, 0, 0, 0, 0, 0, 0, 0, 0, 0, 0, 0, 0, 0, 0, 0, 0 }		},	// ME_REFERENCE
+		{ true,		{ 0, 0, 0, 0, 0, 0, 0, 0, 0, 0, 0, 0, 0, 0, 0, 0, 0, 0, 0, 0, 0, 0, 0, 0, 0, 0, 0, 0, 0 }		},	// ME_CONDITIONAL
+		{ true,		{ 0, 0, 0, 0, 0, 0, 0, 0, 0, 0, 0, 0, 0, 0, 0, 0, 0, 0, 0, 0, 0, 0, 0, 0, 0, 0, 0, 0, 0 }		},	// ME_SWITCH
 		{ true,		{ 0, 0, 0, 0, 0, 0, 0, 0, 0, 0, 0, 0, 0, 0, 0, 0, 0, 0, 0, 0, 0, 0, 0, 0, 0, 0, 0, 0, 0 }		},	// ME_APPLYLAYOUT
 		{ true,		{ 0, 0, 0, 0, 0, 0, 0, 0, 0, 0, 0, 0, 0, 0, 0, 0, 0, 0, 0, 0, 0, 0, 0, 0, 0, 0, 0, 0, 0 }		},	// ME_PREPARELAYOUT
 		{ true,		{ 0, 0, 0, 0, 0, 0, 0, 0, 0, 0, 0, 0, 0, 0, 0, 0, 0, 0, 0, 0, 0, 0, 0, 0, 0, 0, 0, 0, 0 }		},	// ME_DIFFERENCE
@@ -180,6 +161,10 @@
 		{ true,		{ 0, 0, 0, 0, 0, 0, 0, 0, 0, 0, 0, 0, 0, 0, 0, 0, 0, 0, 0, 0, 0, 0, 0, 0, 0, 0, 0, 0, 0 }		},	// ME_EXTERNAL
 		{ true,		{ 0, 0, 0, 0, 0, 0, 0, 0, 0, 0, 0, 0, 0, 0, 0, 0, 0, 0, 0, 0, 0, 0, 0, 0, 0, 0, 0, 0, 0 }		},	// ME_SKELETALMESH_BREAK
 
+		{ false,	{ 0, 0, 0, 0, 0, 0, 0, 0, 0, 0, 0, 0, 0, 0, 0, 0, 0, 0, 0, 0, 0, 0, 0, 0, 0, 0, 0, 0, 0 }		},	// SK_PARAMETER
+
+		{ false,	{ 0, 0, 0, 0, 0, 0, 0, 0, 0, 0, 0, 0, 0, 0, 0, 0, 0, 0, 0, 0, 0, 0, 0, 0, 0, 0, 0, 0, 0 }		},	// IN_CONDITIONAL
+		{ false,	{ 0, 0, 0, 0, 0, 0, 0, 0, 0, 0, 0, 0, 0, 0, 0, 0, 0, 0, 0, 0, 0, 0, 0, 0, 0, 0, 0, 0, 0 }		},	// IN_SWITCH
 		{ false,	{ 0, 0, 0, 0, 0, 0, 0, 0, 0, 0, 0, 0, 0, 0, 0, 0, 0, 0, 0, 0, 0, 0, 0, 0, 0, 0, 0, 0, 0 }		},	// IN_ADDMESH
 		{ false,	{ 0, 0, 0, 0, 0, 0, 0, 0, 0, 0, 0, 0, 0, 0, 0, 0, 0, 0, 0, 0, 0, 0, 0, 0, 0, 0, 0, 0, 0 }		},	// IN_ADDSTRING
 		{ false,	{ 0, 0, 0, 0, 0, 0, 0, 0, 0, 0, 0, 0, 0, 0, 0, 0, 0, 0, 0, 0, 0, 0, 0, 0, 0, 0, 0, 0, 0 }		},	// IN_ADDSURFACE
@@ -191,14 +176,31 @@
 		{ false,	{ 0, 0, 0, 0, 0, 0, 0, 0, 0, 0, 0, 0, 0, 0, 0, 0, 0, 0, 0, 0, 0, 0, 0, 0, 0, 0, 0, 0, 0 }		},	// IN_ADDOVERRIDEMATERIAL
 		{ false,	{ 0, 0, 0, 0, 0, 0, 0, 0, 0, 0, 0, 0, 0, 0, 0, 0, 0, 0, 0, 0, 0, 0, 0, 0, 0, 0, 0, 0, 0 }		},	// IN_ADDMATERIAL
 
+		{ false,	{ 0, 0, 0, 0, 0, 0, 0, 0, 0, 0, 0, 0, 0, 0, 0, 0, 0, 0, 0, 0, 0, 0, 0, 0, 0, 0, 0, 0, 0 }		},	// LA_CONSTANT
+		{ false,	{ 0, 0, 0, 0, 0, 0, 0, 0, 0, 0, 0, 0, 0, 0, 0, 0, 0, 0, 0, 0, 0, 0, 0, 0, 0, 0, 0, 0, 0 }		},	// LA_CONDITIONAL
+		{ false,	{ 0, 0, 0, 0, 0, 0, 0, 0, 0, 0, 0, 0, 0, 0, 0, 0, 0, 0, 0, 0, 0, 0, 0, 0, 0, 0, 0, 0, 0 }		},	// LA_SWITCH
 		{ true,		{ 0, 0, 0, 0, 0, 0, 0, 0, 0, 0, 0, 0, 0, 0, 0, 0, 0, 0, 0, 0, 0, 0, 0, 0, 0, 0, 0, 0, 0 }		},	// LA_PACK
 		{ true,		{ 0, 0, 0, 0, 0, 0, 0, 0, 0, 0, 0, 0, 0, 0, 0, 0, 0, 0, 0, 0, 0, 0, 0, 0, 0, 0, 0, 0, 0 }		},	// LA_MERGE
 		{ true,		{ 0, 0, 0, 0, 0, 0, 0, 0, 0, 0, 0, 0, 0, 0, 0, 0, 0, 0, 0, 0, 0, 0, 0, 0, 0, 0, 0, 0, 0 }		},	// LA_REMOVEBLOCKS
 		{ true,		{ 0, 0, 0, 0, 0, 0, 0, 0, 0, 0, 0, 0, 0, 0, 0, 0, 0, 0, 0, 0, 0, 0, 0, 0, 0, 0, 0, 0, 0 }		},	// LA_FROMMESH
 
-		{ true,		{ 0, 0, 0, 0, 0, 0, 0, 0, 0, 0, 0, 0, 0, 0, 0, 0, 0, 0, 0, 0, 0, 0, 0, 0, 0, 0, 0, 0, 0 }		},	// MI_EXTERNAL
+		{ false,	{ 0, 0, 0, 0, 0, 0, 0, 0, 0, 0, 0, 0, 0, 0, 0, 0, 0, 0, 0, 0, 0, 0, 0, 0, 0, 0, 0, 0, 0 }		},	// MI_CONSTANT
+		{ false,	{ 0, 0, 0, 0, 0, 0, 0, 0, 0, 0, 0, 0, 0, 0, 0, 0, 0, 0, 0, 0, 0, 0, 0, 0, 0, 0, 0, 0, 0 }		},	// MI_PARAMETER
+		{ false,	{ 0, 0, 0, 0, 0, 0, 0, 0, 0, 0, 0, 0, 0, 0, 0, 0, 0, 0, 0, 0, 0, 0, 0, 0, 0, 0, 0, 0, 0 }		},	// MI_CONDITIONAL
+		{ false,	{ 0, 0, 0, 0, 0, 0, 0, 0, 0, 0, 0, 0, 0, 0, 0, 0, 0, 0, 0, 0, 0, 0, 0, 0, 0, 0, 0, 0, 0 }		},	// MI_SWITCH
     	{ true,		{ 0, 0, 0, 0, 0, 0, 0, 0, 0, 0, 0, 0, 0, 0, 0, 0, 0, 0, 0, 0, 0, 0, 0, 0, 0, 0, 0, 0, 0 }		},	// MI_SKELETALMESH_BREAK
+    	{ true,		{ 0, 0, 0, 0, 0, 0, 0, 0, 0, 0, 0, 0, 0, 0, 0, 0, 0, 0, 0, 0, 0, 0, 0, 0, 0, 0, 0, 0, 0 }		},	// MI_FROM_SKELETALMESH_SLOT
+		{ false,	{ 0, 0, 0, 0, 0, 0, 0, 0, 0, 0, 0, 0, 0, 0, 0, 0, 0, 0, 0, 0, 0, 0, 0, 0, 0, 0, 0, 0, 0 }		},	// MI_MODIFY
+		{ true,		{ 0, 0, 0, 0, 0, 0, 0, 0, 0, 0, 0, 0, 0, 0, 0, 0, 0, 0, 0, 0, 0, 0, 0, 0, 0, 0, 0, 0, 0 }		},	// MI_EXTERNAL
 
+		{ false,	{ 0, 0, 0, 0, 0, 0, 0, 0, 0, 0, 0, 0, 0, 0, 0, 0, 0, 0, 0, 0, 0, 0, 0, 0, 0, 0, 0, 0, 0 }		},	// MA_CONSTANT
+		{ false,	{ 0, 0, 0, 0, 0, 0, 0, 0, 0, 0, 0, 0, 0, 0, 0, 0, 0, 0, 0, 0, 0, 0, 0, 0, 0, 0, 0, 0, 0 }		},	// MA_PARAMETER
+		
+		{ true,		{ 0, 0, 0, 0, 0, 0, 0, 0, 0, 0, 0, 0, 0, 0, 0, 0, 0, 0, 0, 0, 0, 0, 0, 0, 0, 0, 0, 0, 0 }		},	// ED_CONSTANT
+		{ false,	{ 0, 0, 0, 0, 0, 0, 0, 0, 0, 0, 0, 0, 0, 0, 0, 0, 0, 0, 0, 0, 0, 0, 0, 0, 0, 0, 0, 0, 0 }		},	// ED_CONDITIONAL
+		{ false,	{ 0, 0, 0, 0, 0, 0, 0, 0, 0, 0, 0, 0, 0, 0, 0, 0, 0, 0, 0, 0, 0, 0, 0, 0, 0, 0, 0, 0, 0 }		},	// ED_SWITCH
+
+		{ false,	{ 0, 0, 0, 0, 0, 0, 0, 0, 0, 0, 0, 0, 0, 0, 0, 0, 0, 0, 0, 0, 0, 0, 0, 0, 0, 0, 0, 0, 0 }		},	// IS_PARAMETER
     	{ true,		{ 0, 0, 0, 0, 0, 0, 0, 0, 0, 0, 0, 0, 0, 0, 0, 0, 0, 0, 0, 0, 0, 0, 0, 0, 0, 0, 0, 0, 0 }		},	// IS_EXTERNAL
 	};
 
@@ -420,7 +422,7 @@
     }
 
     //---------------------------------------------------------------------------------------------
-	TSharedPtr<FModel> Compiler::Compile( const Ptr<Node>& pNode, TMap<OP::ADDRESS, FInstancedStruct>& ExternalOperations)
+	TSharedPtr<FModel> Compiler::Compile( const Ptr<Node>& pNode, TMap<FOperation::ADDRESS, FInstancedStruct>& ExternalOperations)
     {
         MUTABLE_CPUPROFILER_SCOPE(Compile);
 
@@ -548,7 +550,6 @@
 
 		// Preallocate ample memory
 		Program.ByteCode.Reserve(16 * 1024 * 1024);
-		Program.OpAddress.Reserve(1024 * 1024);
 
 		// Keep the link options outside the scope because it is also used to cache constant data that has already been 
 		// added and could be reused across states.
@@ -572,7 +573,6 @@
         }
 
 		Program.ByteCode.Shrink();
-		Program.OpAddress.Shrink();
 
         // Set the runtime parameter indices.
         for(FStateCompilationData& State : States )
@@ -608,7 +608,7 @@
             // Generate the mask of update cache ops
             for (const Ptr<ASTOp>& Instruction : State.m_updateCache )
             {
-                State.state.m_updateCache.Emplace(Instruction->linkedAddress);
+                State.state.m_updateCache.Emplace(Instruction->LinkedAddress);
             }
 
             // Sort the update cache addresses for performance and determinism
@@ -646,7 +646,7 @@
                 //check(relevantMask!=0);
                 if (RelevantMask!=0)
                 {
-                    State.state.m_dynamicResources.Emplace( Instruction.Key->linkedAddress, RelevantMask );
+                    State.state.m_dynamicResources.Emplace(Instruction.Key->LinkedAddress, RelevantMask);
                 }
             }            
 
@@ -687,12 +687,12 @@
 		UE_LOG(LogMutableCore, Log, TEXT("Compilation Time: %lfs"), FPlatformTime::Seconds() - StartTimeSeconds);
 
 		const double ProgramSize  = double(ProgramSizeStream.GetBufferSize());
-		const double ByteCodeSize = double(Program.ByteCode.Num() + Program.OpAddress.Num() * sizeof(uint32));
+		const double ByteCodeSize = double(Program.ByteCode.Num());
 		constexpr double ByteToMiB = 1.0/(1024.0*1024.0); 
 
 		UE_LOG(LogMutableCore, Log, TEXT("Program Size : %lf MiB"), ProgramSize * ByteToMiB);
 		UE_LOG(LogMutableCore, Log, TEXT("Program ByteCode Size: %lf MiB"), ByteCodeSize * ByteToMiB);
-		UE_LOG(LogMutableCore, Log, TEXT("Program Num Ops : %d"), Program.OpAddress.Num());
+		UE_LOG(LogMutableCore, Log, TEXT("Program Num Ops : %d"), Program.NumOps);
 
         return Result;
     }
diff --git a/Engine/Plugins/Mutable/Source/MutableTools/Private/MuT/ErrorLog.cpp b/Engine/Plugins/Mutable/Source/MutableTools/Private/MuT/ErrorLog.cpp
--- a/Engine/Plugins/Mutable/Source/MutableTools/Private/MuT/ErrorLog.cpp
+++ b/Engine/Plugins/Mutable/Source/MutableTools/Private/MuT/ErrorLog.cpp
@@ -189,154 +189,156 @@ namespace UE::Mutable::Private
 
     const TCHAR* s_opNames[] =
 	{
-		TEXT("NONE             "),
-
-		TEXT("BO_CONSTANT      "),
-        TEXT("NU_CONSTANT      "),
-        TEXT("SC_CONSTANT      "),
-        TEXT("CO_CONSTANT      "),
-		TEXT("IM_CONSTANT      "),
-		TEXT("ME_CONSTANT      "),
-		TEXT("LA_CONSTANT      "),
-		TEXT("PR_CONSTANT      "),
-		TEXT("ST_CONSTANT      "),
-		TEXT("ED_CONSTANT      "),
-		TEXT("MA_CONSTANT      "),
-		TEXT("MI_CONSTANT      "),
-
-		TEXT("BO_PARAMETER     "),
-		TEXT("NU_PARAMETER     "),
-		TEXT("SC_PARAMETER     "),
-		TEXT("CO_PARAMETER     "),
-		TEXT("PR_PARAMETER     "),
-		TEXT("IM_PARAMETER     "),
-		TEXT("SK_PARAMETER     "),
-		TEXT("ST_PARAMETER     "),
-		TEXT("MA_PARAMETER     "),
-		TEXT("MI_PARAMETER     "),
-		TEXT("IS_PARAMETER     "),
-		
-		TEXT("IM_REFERENCE     "),
-		TEXT("ME_REFERENCE     "),
-    	
-		TEXT("NU_CONDITIONAL   "),
-		TEXT("SC_CONDITIONAL   "),
-		TEXT("CO_CONDITIONAL   "),
-		TEXT("IM_CONDITIONAL   "),
-		TEXT("ME_CONDITIONAL   "),
-		TEXT("LA_CONDITIONAL   "),
-		TEXT("IN_CONDITIONAL   "),
-		TEXT("ED_CONDITIONAL   "),
-		TEXT("MI_CONDITIONAL   "),
-		
-		TEXT("NU_SWITCH        "),
-		TEXT("SC_SWITCH        "),
-		TEXT("CO_SWITCH        "),
-		TEXT("IM_SWITCH        "),
-		TEXT("ME_SWITCH        "),
-		TEXT("LA_SWITCH        "),
-		TEXT("IN_SWITCH        "),
-		TEXT("ED_SWITCH        "),
-		TEXT("MI_SWITCH        "),
-
+		TEXT("NONE"),
+
+		TEXT("BO_CONSTANT"),
+		TEXT("BO_PARAMETER"),
+		TEXT("BO_EQUAL_INT_CONST"),
+		TEXT("BO_AND"),
+		TEXT("BO_OR"),
+		TEXT("BO_NOT"),
+
+		TEXT("NU_CONSTANT"),
+		TEXT("NU_PARAMETER"),
+		TEXT("NU_CONDITIONAL"),
+		TEXT("NU_SWITCH"),
+
+		TEXT("SC_CONSTANT"),
+		TEXT("SC_PARAMETER"),
+		TEXT("SC_CONDITIONAL"),
+		TEXT("SC_SWITCH"),
 		TEXT("SC_MATERIAL_BREAK"),
+		TEXT("SC_ARITHMETIC"),
+		TEXT("SC_CURVE"),
+		TEXT("SC_EXTERNAL"),
+
+		TEXT("CO_CONSTANT"),
+		TEXT("CO_PARAMETER"),
+		TEXT("CO_CONDITIONAL"),
+		TEXT("CO_SWITCH"),
 		TEXT("CO_MATERIAL_BREAK"),
+		TEXT("CO_SAMPLEIMAGE"),
+		TEXT("CO_SWIZZLE"),
+		TEXT("CO_FROMSCALARS"),
+		TEXT("CO_ARITHMETIC"),
+		TEXT("CO_LINEARTOSRGB"),
+		TEXT("CO_EXTERNAL"),
+
+		TEXT("ST_CONSTANT"),
+		TEXT("ST_PARAMETER"),
+
+		TEXT("PR_CONSTANT"),
+		TEXT("PR_PARAMETER"),
+
+		TEXT("IM_CONSTANT"),
+		TEXT("IM_PARAMETER"),
+		TEXT("IM_REFERENCE"),
+		TEXT("IM_CONDITIONAL"),
+		TEXT("IM_SWITCH"),
 		TEXT("IM_MATERIAL_BREAK"),
 		TEXT("IM_PARAMETER_FROM_MATERIAL"),
-		TEXT("MI_MODIFY        "),
-
-		TEXT("BO_EQUAL_SC_CONST"),
-		TEXT("BO_AND           "),
-		TEXT("BO_OR            "),
-		TEXT("BO_NOT           "),
-		
-		TEXT("SC_ARITHMETIC    "),
-		TEXT("SC_CURVE         "),
-    	TEXT("SC_EXTERNAL      "),
-
-		TEXT("CO_SAMPLEIMAGE   "),
-		TEXT("CO_SWIZZLE       "),
-		TEXT("CO_FROMSCALARS   "),
-		TEXT("CO_ARITHMETIC    "),
-		TEXT("CO_LINEARTOSRGB  "),
-    	TEXT("CO_EXTERNAL      "),
-
-		TEXT("IM_LAYER         "),
-		TEXT("IM_LAYERCOLOUR   "),
-		TEXT("IM_PIXELFORMAT   "),
-		TEXT("IM_MIPMAP        "),
-		TEXT("IM_RESIZE        "),
-		TEXT("IM_RESIZELIKE    "),
-		TEXT("IM_RESIZEREL     "),
-		TEXT("IM_BLANKLAYOUT   "),
-		TEXT("IM_COMPOSE       "),
-		TEXT("IM_MULTICOMPOSE  "),
-		TEXT("IM_INTERPOLATE   "),
-		TEXT("IM_SATURATE      "),
-		TEXT("IM_LUMINANCE     "),
-		TEXT("IM_SWIZZLE       "),
-		TEXT("IM_COLOURMAP     "),
-		TEXT("IM_BINARISE      "),
-		TEXT("IM_PLAINCOLOUR   "),
-		TEXT("IM_CROP          "),
-		TEXT("IM_PATCH         "),
-		TEXT("IM_RASTERMESH    "),
-		TEXT("IM_MAKEGROWMAP   "),
-		TEXT("IM_DISPLACE      "),
-		TEXT("IM_MULTILAYER    "),
-		TEXT("IM_INVERT        "),
-		TEXT("IM_NORMAL_COMPO  "),
-		TEXT("IM_TRANSFORM     "),
+		TEXT("IM_LAYER"),
+		TEXT("IM_LAYERCOLOUR"),
+		TEXT("IM_PIXELFORMAT"),
+		TEXT("IM_MIPMAP"),
+		TEXT("IM_RESIZE"),
+		TEXT("IM_RESIZELIKE"),
+		TEXT("IM_RESIZEREL"),
+		TEXT("IM_BLANKLAYOUT"),
+		TEXT("IM_COMPOSE"),
+		TEXT("IM_MULTICOMPOSE"),
+		TEXT("IM_INTERPOLATE"),
+		TEXT("IM_SATURATE"),
+		TEXT("IM_LUMINANCE"),
+		TEXT("IM_SWIZZLE"),
+		TEXT("IM_COLOURMAP"),
+		TEXT("IM_BINARISE"),
+		TEXT("IM_PLAINCOLOUR"),
+		TEXT("IM_CROP"),
+		TEXT("IM_PATCH"),
+		TEXT("IM_RASTERMESH"),
+		TEXT("IM_MAKEGROWMAP"),
+		TEXT("IM_DISPLACE"),
+		TEXT("IM_MULTILAYER"),
+		TEXT("IM_INVERT"),
+		TEXT("IM_NORMALCOMPOSITE"),
+		TEXT("IM_TRANSFORM"),
 		TEXT("IM_PARAMETER_CONVERT"),
-    	TEXT("IM_EXTERNAL      "),
-
-		TEXT("ME_APPLYLAYOUT   "),
-		TEXT("ME_PREPARELAYOUT "),
-		TEXT("ME_DIFFERENCE    "),
-		TEXT("ME_MORPH         "),
-		TEXT("ME_MERGE         "),
-		TEXT("ME_MASKCLIPMESH  "),
+		TEXT("IM_EXTERNAL"),
+
+		TEXT("ME_CONSTANT"),
+		TEXT("ME_PARAMETER"),
+		TEXT("ME_REFERENCE"),
+		TEXT("ME_CONDITIONAL"),
+		TEXT("ME_SWITCH"),
+		TEXT("ME_APPLYLAYOUT"),
+		TEXT("ME_PREPARELAYOUT"),
+		TEXT("ME_DIFFERENCE"),
+		TEXT("ME_MORPH"),
+		TEXT("ME_MERGE"),
+		TEXT("ME_MASKCLIPMESH"),
 		TEXT("ME_MASKCLIPUVMASK"),
-		TEXT("ME_MASKDIFF      "),
-        TEXT("ME_REMOVEMASK    "),
-        TEXT("ME_FORMAT        "),
-        TEXT("ME_EXTRACTLABLOCK"),
-		TEXT("ME_TRANSFORM     "),
+		TEXT("ME_MASKDIFF"),
+		TEXT("ME_REMOVEMASK"),
+		TEXT("ME_FORMAT"),
+		TEXT("ME_EXTRACTLAYOUTBLOCK"),
+		TEXT("ME_TRANSFORM"),
 		TEXT("ME_CLIPMORPHPLANE"),
-        TEXT("ME_CLIPWITHMESH  "),
-        TEXT("ME_PROJECT       "),
-        TEXT("ME_APPLYPOSE     "),
-		TEXT("ME_BINDSHAPE	   "),
-		TEXT("ME_APPLYSHAPE	   "),
-		TEXT("ME_CLIPDEFORM	   "),
-		TEXT("ME_MORPHRESHAPE  "),
-		TEXT("ME_OPTIMIZESKIN  "),
-		TEXT("ME_ADDMETADATA   "),
+		TEXT("ME_CLIPWITHMESH"),
+		TEXT("ME_PROJECT"),
+		TEXT("ME_APPLYPOSE"),
+		TEXT("ME_BINDSHAPE"),
+		TEXT("ME_APPLYSHAPE"),
+		TEXT("ME_CLIPDEFORM"),
+		TEXT("ME_MORPHRESHAPE"),
+		TEXT("ME_OPTIMIZESKIN"),
+		TEXT("ME_ADDMETADATA"),
 		TEXT("ME_TRANSFORMWITHMESH"),
 		TEXT("ME_TRANSFORMWITHBONE"),
-    	TEXT("ME_EXTERNAL      "),
+		TEXT("ME_EXTERNAL"),
 		TEXT("ME_SKELETALMESH_BREAK"),
 
-		TEXT("IN_ADDMESH       "),
-		TEXT("IN_ADDSTRING     "),
-        TEXT("IN_ADDSURFACE    "),
-        TEXT("IN_ADDCOMPONENT  "),
-        TEXT("IN_ADDLOD        "),
-        TEXT("IN_ADDSKELETALMESH"),
-		TEXT("IN_ADDEXTENSIDATA"),
+		TEXT("SK_PARAMETER"),
+
+		TEXT("IN_CONDITIONAL"),
+		TEXT("IN_SWITCH"),
+		TEXT("IN_ADDMESH"),
+		TEXT("IN_ADDSTRING"),
+		TEXT("IN_ADDSURFACE"),
+		TEXT("IN_ADDCOMPONENT"),
+		TEXT("IN_ADDLOD"),
+		TEXT("IN_ADDSKELETALMESH"),
+		TEXT("IN_ADDEXTENSIONDATA"),
 		TEXT("IN_ADDOVERLAYMATERIAL"),
 		TEXT("IN_ADDOVERRIDEMATERIAL"),
-		TEXT("IN_ADDMATERIAL   "),
-		
-		TEXT("LA_PACK          "),
-		TEXT("LA_MERGE         "),
-		TEXT("LA_REMOVEBLOCKS  "),
-		TEXT("LA_FROMMESH      "),
-
-    	TEXT("MI_EXTERNAL      "),
+		TEXT("IN_ADDMATERIAL"),
+
+		TEXT("LA_CONSTANT"),
+		TEXT("LA_CONDITIONAL"),
+		TEXT("LA_SWITCH"),
+		TEXT("LA_PACK"),
+		TEXT("LA_MERGE"),
+		TEXT("LA_REMOVEBLOCKS"),
+		TEXT("LA_FROMMESH"),
+
+		TEXT("MI_CONSTANT"),
+		TEXT("MI_PARAMETER"),
+		TEXT("MI_CONDITIONAL"),
+		TEXT("MI_SWITCH"),
+		TEXT("MI_SKELETALMESH_BREAK"),
     	TEXT("MI_FROM_SKELETALMESH_SLOT"),
+		TEXT("MI_MODIFY"),
+		TEXT("MI_EXTERNAL"),
+
+		TEXT("MA_CONSTANT"),
+		TEXT("MA_PARAMETER"),
+	
+		TEXT("ED_CONSTANT"),
+		TEXT("ED_CONDITIONAL"),
+		TEXT("ED_SWITCH"),
 
-    	TEXT("IS_EXTERNAL      "),
+		TEXT("IS_PARAMETER"),
+		TEXT("IS_EXTERNAL"),
 	};
 
 	static_assert(UE_ARRAY_COUNT(s_opNames) == int32(EOpType::COUNT));
