Fix static analysis warning.

#rnx
#Synced-CL 51454919

diff --git a/Engine/Plugins/Mutable/Source/MutableRuntime/Private/MuR/ProgramCache.cpp b/Engine/Plugins/Mutable/Source/MutableRuntime/Private/MuR/ProgramCache.cpp
--- a/Engine/Plugins/Mutable/Source/MutableRuntime/Private/MuR/ProgramCache.cpp
+++ b/Engine/Plugins/Mutable/Source/MutableRuntime/Private/MuR/ProgramCache.cpp
@@ -359,8 +359,8 @@ namespace ProgramCacheInternals
 				}
 
 				const bool bFreeAddress = 
-					(  bIsLockedAddress  & EnumHasAnyFlags(ClearFlags, FProgramCache::EClearFlags::Locked)) | 
-					((!bIsLockedAddress) & EnumHasAnyFlags(ClearFlags, FProgramCache::EClearFlags::Unlocked));
+					(  bIsLockedAddress  && EnumHasAnyFlags(ClearFlags, FProgramCache::EClearFlags::Locked)) || 
+					((!bIsLockedAddress) && EnumHasAnyFlags(ClearFlags, FProgramCache::EClearFlags::Unlocked));
 
 				if (bFreeAddress)
 				{
