Minor Mutable dependency cleanup.

#Synced-CL 51451821

diff --git a/Engine/Plugins/Mutable/Source/CustomizableObjectEditor/Private/MuCOE/GenerateMutableSource/GenerateMutableSourceImage.cpp b/Engine/Plugins/Mutable/Source/CustomizableObjectEditor/Private/MuCOE/GenerateMutableSource/GenerateMutableSourceImage.cpp
--- a/Engine/Plugins/Mutable/Source/CustomizableObjectEditor/Private/MuCOE/GenerateMutableSource/GenerateMutableSourceImage.cpp
+++ b/Engine/Plugins/Mutable/Source/CustomizableObjectEditor/Private/MuCOE/GenerateMutableSource/GenerateMutableSourceImage.cpp
@@ -4,6 +4,7 @@
 
 #include "GenerateMutableSourceExternal.h"
 #include "Engine/TextureLODSettings.h"
+#include "Modules/ModuleManager.h"
 
 #include "MuCOE/CustomizableObjectCompiler.h"
 #include "MuCOE/EdGraphSchema_CustomizableObject.h"
diff --git a/Engine/Plugins/Mutable/Source/CustomizableObjectEditor/Private/MuCOE/Nodes/CONodeMaterialBreak.cpp b/Engine/Plugins/Mutable/Source/CustomizableObjectEditor/Private/MuCOE/Nodes/CONodeMaterialBreak.cpp
--- a/Engine/Plugins/Mutable/Source/CustomizableObjectEditor/Private/MuCOE/Nodes/CONodeMaterialBreak.cpp
+++ b/Engine/Plugins/Mutable/Source/CustomizableObjectEditor/Private/MuCOE/Nodes/CONodeMaterialBreak.cpp
@@ -3,6 +3,7 @@
 #include "MuCOE/Nodes/CONodeMaterialBreak.h"
 
 #include "DetailsViewArgs.h"
+#include "Modules/ModuleManager.h"
 #include "MuCOE/EdGraphSchema_CustomizableObject.h"
 #include "PropertyEditorModule.h"
 #include "ScopedTransaction.h"
diff --git a/Engine/Plugins/Mutable/Source/CustomizableObjectEditor/Private/MuCOE/Nodes/CONodeMaterialModify.cpp b/Engine/Plugins/Mutable/Source/CustomizableObjectEditor/Private/MuCOE/Nodes/CONodeMaterialModify.cpp
--- a/Engine/Plugins/Mutable/Source/CustomizableObjectEditor/Private/MuCOE/Nodes/CONodeMaterialModify.cpp
+++ b/Engine/Plugins/Mutable/Source/CustomizableObjectEditor/Private/MuCOE/Nodes/CONodeMaterialModify.cpp
@@ -3,6 +3,7 @@
 #include "MuCOE/Nodes/CONodeMaterialModify.h"
 
 #include "DetailsViewArgs.h"
+#include "Modules/ModuleManager.h"
 #include "MuCOE/EdGraphSchema_CustomizableObject.h"
 #include "PropertyEditorModule.h"
 #include "ScopedTransaction.h"
diff --git a/Engine/Plugins/Mutable/Source/MutableRuntime/Internal/MuR/Registry.h b/Engine/Plugins/Mutable/Source/MutableRuntime/Internal/MuR/Registry.h
--- a/Engine/Plugins/Mutable/Source/MutableRuntime/Internal/MuR/Registry.h
+++ b/Engine/Plugins/Mutable/Source/MutableRuntime/Internal/MuR/Registry.h
@@ -2,8 +2,6 @@
 
 #pragma once
 
-#include "Chaos/Character/CharacterGroundConstraintContainer.h"
-#include "Chaos/Deformable/MuscleActivationConstraints.h"
 #include "Templates/SharedPointer.h"
 #include "Misc/ScopeLock.h"
 
diff --git a/Engine/Plugins/Mutable/Source/MutableRuntime/Private/MuR/OpImageMipmap.cpp b/Engine/Plugins/Mutable/Source/MutableRuntime/Private/MuR/OpImageMipmap.cpp
--- a/Engine/Plugins/Mutable/Source/MutableRuntime/Private/MuR/OpImageMipmap.cpp
+++ b/Engine/Plugins/Mutable/Source/MutableRuntime/Private/MuR/OpImageMipmap.cpp
@@ -5,7 +5,7 @@
 #include "MuR/System.h"
 #include "MuR/ParallelExecutionUtils.h"
 #include "HAL/UnrealMemory.h"
-
+#include "Async/ParallelFor.h"
 
 namespace
 {
