[Mutable] Added AssetUserData type and nodes.

- New nodes:
- Switch
- Parameter
- Mesh Add Asset User Data
- Removed per type Switch classes.
- Automatically add Switch nodes to External types.
- Allow adding custom pin color to Extenral types.
- Changed External Operation Output/Input definitions from FName to FText.
- Removed External Operations GetName. Replaced with Display Name specifier.
- Fixed multiple issues with new extensions system.

#Synced-CL 51468966

diff --git a/Engine/Plugins/Experimental/MutableAssetUserData/MutableAssetUserData.uplugin b/Engine/Plugins/Experimental/MutableAssetUserData/MutableAssetUserData.uplugin
new file mode 100644
--- /dev/null
+++ b/Engine/Plugins/Experimental/MutableAssetUserData/MutableAssetUserData.uplugin
@@ -0,0 +1,33 @@
+{
+	"FileVersion": 3,
+	"Version": 1,
+	"VersionName": "0.1",
+	"FriendlyName": "Mutable Asset User Data",
+	"Description": "Adds Asset User Data type and operations.",
+	"Category": "Animation",
+	"CreatedBy": "Epic Games, Inc.",
+	"CreatedByURL": "https://epicgames.com",
+	"DocsURL": "",
+	"MarketplaceURL": "",
+	"SupportURL": "",
+	"CanContainContent": false,
+	"IsBetaVersion": false,
+	"IsExperimentalVersion": true,
+	"Installed": false,
+	"EnabledByDefault": true,
+	"Modules": 
+	[
+		{
+			"Name": "MutableAssetUserData",
+			"Type": "Runtime",
+			"LoadingPhase": "Default"
+		}
+	],
+	"Plugins": 
+	[
+		{
+			"Name": "Mutable",
+			"Enabled": true
+		}
+	]
+}
\ No newline at end of file
diff --git a/Engine/Plugins/Experimental/MutableAssetUserData/Source/MutableAssetUserData/MutableAssetUserData.Build.cs b/Engine/Plugins/Experimental/MutableAssetUserData/Source/MutableAssetUserData/MutableAssetUserData.Build.cs
new file mode 100644
--- /dev/null
+++ b/Engine/Plugins/Experimental/MutableAssetUserData/Source/MutableAssetUserData/MutableAssetUserData.Build.cs
@@ -0,0 +1,28 @@
+// Copyright Epic Games, Inc. All Rights Reserved.
+
+namespace UnrealBuildTool.Rules
+{
+	public class MutableAssetUserData : ModuleRules
+	{
+        public MutableAssetUserData(ReadOnlyTargetRules Target) : base(Target)
+		{
+			PrivateIncludePaths.Add(ModuleDirectory + "/Private");
+			PublicIncludePaths.Add(ModuleDirectory + "/Public");
+			
+			PublicDependencyModuleNames.AddRange(
+				new string[]
+				{
+					"Engine",
+					"Core",
+					"CoreUObject",
+					"MutableRuntime"
+				});
+
+			PrivateDependencyModuleNames.AddRange(
+				new string[]
+				{
+					"CustomizableObject"
+				});
+		}
+	}
+}
diff --git a/Engine/Plugins/Experimental/MutableAssetUserData/Source/MutableAssetUserData/Private/MeshAddAssetUserData.cpp b/Engine/Plugins/Experimental/MutableAssetUserData/Source/MutableAssetUserData/Private/MeshAddAssetUserData.cpp
new file mode 100644
--- /dev/null
+++ b/Engine/Plugins/Experimental/MutableAssetUserData/Source/MutableAssetUserData/Private/MeshAddAssetUserData.cpp
@@ -0,0 +1,50 @@
+﻿// Copyright Epic Games, Inc. All Rights Reserved.
+
+#include "MeshAddAssetUserData.h"
+
+#include "MutableAssetUserData.h"
+#include "MuR/Mesh.h"
+#include "MuR/External/MeshAdapter.h"
+#include "Engine/AssetUserData.h"
+
+#define LOCTEXT_NAMESPACE "MutableAssetUserData"
+
+
+const FText FMeshAddAssetUserData::TextInputMesh = LOCTEXT("Mesh", "Mesh");
+
+const FText FMeshAddAssetUserData::TextInputAssetUserData = LOCTEXT("AssetUserData", "Asset User Data");
+
+
+TArray<TPair<FText, const UScriptStruct*>> FMeshAddAssetUserData::GetInputs() const
+{
+	TArray<TPair<FText, const UScriptStruct*>> Inputs;
+
+	Inputs.Emplace(TextInputMesh, UE::Mutable::FMeshAdapter::StaticStruct());
+	Inputs.Emplace(TextInputAssetUserData, FMutableAssetUserData::StaticStruct());
+	
+	return Inputs;
+}
+
+
+TPair<FText, const UScriptStruct*> FMeshAddAssetUserData::GetOutput() const
+{
+	return MakeTuple(LOCTEXT("Mesh", "Mesh"), UE::Mutable::FMeshAdapter::StaticStruct());
+}
+
+
+void FMeshAddAssetUserData::Evaluate(UE::Mutable::FContext& Context) const
+{
+	UE::Mutable::FValue InputMesh = UE::Mutable::CopyOrMove(Context.GetInput(TextInputMesh));
+	UE::Mutable::FMeshAdapter& MeshAdapter = InputMesh.Get<UE::Mutable::FMeshAdapter>();
+
+	UE::Mutable::FValueConst InputAssetUserData = Context.GetInput(TextInputAssetUserData);
+	const FMutableAssetUserData& AssetUserData = InputAssetUserData.Get<FMutableAssetUserData>();
+	
+	MeshAdapter.GetPrivate()->AssetUserData.Add(UE::Mutable::Private::TPassthroughObjectPtr<UAssetUserData>(AssetUserData.AssetUserData));
+	
+	// Output the modified Mesh Adapter.
+	Context.SetOutput(MoveTemp(InputMesh));
+}
+
+
+#undef LOCTEXT_NAMESPACE
\ No newline at end of file
diff --git a/Engine/Plugins/Experimental/MutableAssetUserData/Source/MutableAssetUserData/Public/MeshAddAssetUserData.h b/Engine/Plugins/Experimental/MutableAssetUserData/Source/MutableAssetUserData/Public/MeshAddAssetUserData.h
new file mode 100644
--- /dev/null
+++ b/Engine/Plugins/Experimental/MutableAssetUserData/Source/MutableAssetUserData/Public/MeshAddAssetUserData.h
@@ -0,0 +1,23 @@
+﻿// Copyright Epic Games, Inc. All Rights Reserved.
+
+#pragma once
+
+#include "MuR/External/Operation.h"
+
+#include "MeshAddAssetUserData.generated.h"
+
+
+USTRUCT(DisplayName = "Mesh Add Asset User Data")
+struct FMeshAddAssetUserData : public UE::Mutable::FExternalOperation
+{
+	GENERATED_BODY()
+	
+	virtual ~FMeshAddAssetUserData() override = default;
+	virtual TArray<TPair<FText, const UScriptStruct*>> GetInputs() const override;
+	virtual TPair<FText, const UScriptStruct*> GetOutput() const override;
+	virtual void Evaluate(UE::Mutable::FContext& Context) const override;
+	
+	static const FText TextInputMesh;
+	
+	static const FText TextInputAssetUserData;
+};
\ No newline at end of file
diff --git a/Engine/Plugins/Experimental/MutableAssetUserData/Source/MutableAssetUserData/Public/MutableAssetUserData.h b/Engine/Plugins/Experimental/MutableAssetUserData/Source/MutableAssetUserData/Public/MutableAssetUserData.h
new file mode 100644
--- /dev/null
+++ b/Engine/Plugins/Experimental/MutableAssetUserData/Source/MutableAssetUserData/Public/MutableAssetUserData.h
@@ -0,0 +1,18 @@
+﻿// Copyright Epic Games, Inc. All Rights Reserved.
+
+#pragma once
+
+#include "UObject/ObjectPtr.h"
+
+#include "MutableAssetUserData.generated.h"
+
+class UAssetUserData;
+
+USTRUCT(meta = (DisplayName = "Asset User Data", PinColor = "#949494"))
+struct FMutableAssetUserData
+{
+	GENERATED_BODY()
+
+	UPROPERTY(EditAnywhere, Category = AssetUserData)
+	TObjectPtr<UAssetUserData> AssetUserData;
+};
\ No newline at end of file
diff --git a/Engine/Plugins/Experimental/MutableExtensionExample/Source/MutableExtensionExample/Private/MeshClipSphere.cpp b/Engine/Plugins/Experimental/MutableExtensionExample/Source/MutableExtensionExample/Private/MeshClipSphere.cpp
--- a/Engine/Plugins/Experimental/MutableExtensionExample/Source/MutableExtensionExample/Private/MeshClipSphere.cpp
+++ b/Engine/Plugins/Experimental/MutableExtensionExample/Source/MutableExtensionExample/Private/MeshClipSphere.cpp
@@ -5,45 +5,41 @@
 #include "PrimitiveSphere.h"
 #include "MuR/External/MeshAdapter.h"
 
+#define LOCTEXT_NAMESPACE "MutableExtensionExample"
 
-const FName FMeshClipSphere::NameInputMesh = TEXT("Mesh");
 
-const FName FMeshClipSphere::NameInputPrimitiveSphere = TEXT("Sphere");
+const FText FMeshClipSphere::TextInputMesh = LOCTEXT("Mesh", "Mesh");
 
-
-FName FMeshClipSphere::GetName() const
-{
-	return TEXT("Mesh Clip With Sphere");
-}
+const FText FMeshClipSphere::TextInputPrimitiveSphere = LOCTEXT("Sphere", "Sphere");
 
 
-TArray<TPair<FName, const UScriptStruct*>> FMeshClipSphere::GetInputs() const
+TArray<TPair<FText, const UScriptStruct*>> FMeshClipSphere::GetInputs() const
 {
-	TArray<TPair<FName, const UScriptStruct*>> Inputs;
+	TArray<TPair<FText, const UScriptStruct*>> Inputs;
 
-	Inputs.Emplace(NameInputMesh, UE::Mutable::FMeshAdapter::StaticStruct());
-	Inputs.Emplace(NameInputPrimitiveSphere, FPrimitiveSphere::StaticStruct());
+	Inputs.Emplace(TextInputMesh, UE::Mutable::FMeshAdapter::StaticStruct());
+	Inputs.Emplace(TextInputPrimitiveSphere, FPrimitiveSphere::StaticStruct());
 	
 	return Inputs;
 }
 
 
-TPair<FName, const UScriptStruct*> FMeshClipSphere::GetOutput() const
+TPair<FText, const UScriptStruct*> FMeshClipSphere::GetOutput() const
 {
-	return MakeTuple(TEXT("Mesh"), UE::Mutable::FMeshAdapter::StaticStruct());
+	return MakeTuple(LOCTEXT("Mesh", "Mesh"), UE::Mutable::FMeshAdapter::StaticStruct());
 }
 
 
 void FMeshClipSphere::Evaluate(UE::Mutable::FContext& Context) const
 {
 	// Since we are going to edit the Mesh input, we need to make it non-constant using the CopyOrMove.
-	UE::Mutable::FValue InputMesh = UE::Mutable::CopyOrMove(Context.GetInput(NameInputMesh));
+	UE::Mutable::FValue InputMesh = UE::Mutable::CopyOrMove(Context.GetInput(TextInputMesh));
 
 	// Cast the input to a Mesh Adapter.
 	UE::Mutable::FMeshAdapter& Mesh = InputMesh.Get<UE::Mutable::FMeshAdapter>();
 
 	// Get the Primitive input
-	UE::Mutable::FValueConst InputSphere = Context.GetInput(NameInputPrimitiveSphere);
+	UE::Mutable::FValueConst InputSphere = Context.GetInput(TextInputPrimitiveSphere);
 	const FPrimitiveSphere& PrimitiveSphere = InputSphere.Get<FPrimitiveSphere>();
 	
 	TBitArray<> VerticesToClip;
@@ -64,3 +60,6 @@ void FMeshClipSphere::Evaluate(UE::Mutable::FContext& Context) const
 	// Output the modified Mesh Adapter.
 	Context.SetOutput(MoveTemp(InputMesh));
 }
+
+
+#undef LOCTEXT_NAMESPACE
\ No newline at end of file
diff --git a/Engine/Plugins/Experimental/MutableExtensionExample/Source/MutableExtensionExample/Private/MeshIdentity.cpp b/Engine/Plugins/Experimental/MutableExtensionExample/Source/MutableExtensionExample/Private/MeshIdentity.cpp
--- a/Engine/Plugins/Experimental/MutableExtensionExample/Source/MutableExtensionExample/Private/MeshIdentity.cpp
+++ b/Engine/Plugins/Experimental/MutableExtensionExample/Source/MutableExtensionExample/Private/MeshIdentity.cpp
@@ -4,33 +4,32 @@
 
 #include "MuR/External/MeshAdapter.h"
 
+#define LOCTEXT_NAMESPACE "MutableExtensionExample"
 
-const FName FMeshIdentity::NameInputMesh = TEXT("Mesh");
 
-
-FName FMeshIdentity::GetName() const
-{
-	return TEXT("Mesh Identity");
-}
+const FText FMeshIdentity::TextInputMesh = LOCTEXT("Mesh", "Mesh");
 
 
-TArray<TPair<FName, const UScriptStruct*>> FMeshIdentity::GetInputs() const
+TArray<TPair<FText, const UScriptStruct*>> FMeshIdentity::GetInputs() const
 {
-	TArray<TPair<FName, const UScriptStruct*>> Inputs;
+	TArray<TPair<FText, const UScriptStruct*>> Inputs;
 
-	Inputs.Emplace(NameInputMesh, UE::Mutable::FMeshAdapter::StaticStruct());
+	Inputs.Emplace(TextInputMesh, UE::Mutable::FMeshAdapter::StaticStruct());
 	
 	return Inputs;
 }
 
 
-TPair<FName, const UScriptStruct*> FMeshIdentity::GetOutput() const
+TPair<FText, const UScriptStruct*> FMeshIdentity::GetOutput() const
 {
-	return MakeTuple(TEXT("Mesh"), UE::Mutable::FMeshAdapter::StaticStruct());
+	return MakeTuple(LOCTEXT("Mesh", "Mesh"), UE::Mutable::FMeshAdapter::StaticStruct());
 }
 
 
 void FMeshIdentity::Evaluate(UE::Mutable::FContext& Context) const
 {
-	Context.SetOutput(Context.GetInput(NameInputMesh));
+	Context.SetOutput(Context.GetInput(TextInputMesh));
 }
+
+
+#undef LOCTEXT_NAMESPACE
\ No newline at end of file
diff --git a/Engine/Plugins/Experimental/MutableExtensionExample/Source/MutableExtensionExample/Private/MeshNoise.cpp b/Engine/Plugins/Experimental/MutableExtensionExample/Source/MutableExtensionExample/Private/MeshNoise.cpp
--- a/Engine/Plugins/Experimental/MutableExtensionExample/Source/MutableExtensionExample/Private/MeshNoise.cpp
+++ b/Engine/Plugins/Experimental/MutableExtensionExample/Source/MutableExtensionExample/Private/MeshNoise.cpp
@@ -5,45 +5,41 @@
 #include "MuR/External/FloatAdapter.h"
 #include "MuR/External/MeshAdapter.h"
 
+#define LOCTEXT_NAMESPACE "MutableExtensionExample"
 
-const FName FMeshNoise::NameInputMesh = TEXT("Mesh");
 
-const FName FMeshNoise::NameInputFactor = TEXT("Factor");
+const FText FMeshNoise::TextInputMesh = LOCTEXT("Mesh", "Mesh");
 
-
-FName FMeshNoise::GetName() const
-{
-	return TEXT("Mesh Noise");
-}
+const FText FMeshNoise::TextInputFactor = LOCTEXT("Factor", "Factor");
 
 
-TArray<TPair<FName, const UScriptStruct*>> FMeshNoise::GetInputs() const
+TArray<TPair<FText, const UScriptStruct*>> FMeshNoise::GetInputs() const
 {
-	TArray<TPair<FName, const UScriptStruct*>> Inputs;
+	TArray<TPair<FText, const UScriptStruct*>> Inputs;
 
-	Inputs.Emplace(NameInputMesh, UE::Mutable::FMeshAdapter::StaticStruct());
-	Inputs.Emplace(NameInputFactor, UE::Mutable::FFloatAdapter::StaticStruct());
+	Inputs.Emplace(TextInputMesh, UE::Mutable::FMeshAdapter::StaticStruct());
+	Inputs.Emplace(TextInputFactor, UE::Mutable::FFloatAdapter::StaticStruct());
 	
 	return Inputs;
 }
 
 
-TPair<FName, const UScriptStruct*> FMeshNoise::GetOutput() const
+TPair<FText, const UScriptStruct*> FMeshNoise::GetOutput() const
 {
-	return MakeTuple(TEXT("Mesh"), UE::Mutable::FMeshAdapter::StaticStruct());
+	return MakeTuple(LOCTEXT("Mesh", "Mesh"), UE::Mutable::FMeshAdapter::StaticStruct());
 }
 
 
 void FMeshNoise::Evaluate(UE::Mutable::FContext& Context) const
 {
 	// Since we are going to edit the Mesh input, we need to make it non-constant using the CopyOrMove.
-	UE::Mutable::FValue Input = UE::Mutable::CopyOrMove(Context.GetInput(NameInputMesh));
+	UE::Mutable::FValue Input = UE::Mutable::CopyOrMove(Context.GetInput(TextInputMesh));
 
 	// Cast the input to a Mesh Adapter.
 	UE::Mutable::FMeshAdapter& Mesh = Input.Get<UE::Mutable::FMeshAdapter>();
 
 	// Get the Factor input
-	const float Factor = Context.GetInput(NameInputFactor).Get<UE::Mutable::FFloatAdapter>().GetValue();
+	const float Factor = Context.GetInput(TextInputFactor).Get<UE::Mutable::FFloatAdapter>().GetValue();
 	
 	// Move each vertex some random noise. This random is allowed since it is deterministic.
 	FRandomStream Random(Seed);
@@ -60,3 +56,6 @@ void FMeshNoise::Evaluate(UE::Mutable::FContext& Context) const
 	// Output the modified Mesh Adapter.
 	Context.SetOutput(MoveTemp(Input));
 }
+
+
+#undef LOCTEXT_NAMESPACE
\ No newline at end of file
diff --git a/Engine/Plugins/Experimental/MutableExtensionExample/Source/MutableExtensionExample/Public/MeshClipSphere.h b/Engine/Plugins/Experimental/MutableExtensionExample/Source/MutableExtensionExample/Public/MeshClipSphere.h
--- a/Engine/Plugins/Experimental/MutableExtensionExample/Source/MutableExtensionExample/Public/MeshClipSphere.h
+++ b/Engine/Plugins/Experimental/MutableExtensionExample/Source/MutableExtensionExample/Public/MeshClipSphere.h
@@ -7,18 +7,17 @@
 #include "MeshClipSphere.generated.h"
 
 
-USTRUCT()
+USTRUCT(DisplayName = "Mesh Clip With Sphere")
 struct FMeshClipSphere : public UE::Mutable::FExternalOperation
 {
 	GENERATED_BODY()
 	
 	virtual ~FMeshClipSphere() override = default;
-	virtual FName GetName() const override;
-	virtual TArray<TPair<FName, const UScriptStruct*>> GetInputs() const override;
-	virtual TPair<FName, const UScriptStruct*> GetOutput() const override;
+	virtual TArray<TPair<FText, const UScriptStruct*>> GetInputs() const override;
+	virtual TPair<FText, const UScriptStruct*> GetOutput() const override;
 	virtual void Evaluate(UE::Mutable::FContext& Context) const override;
 
-	const static FName NameInputMesh;
+	const static FText TextInputMesh;
 
-	const static FName NameInputPrimitiveSphere;
+	const static FText TextInputPrimitiveSphere;
 };
\ No newline at end of file
diff --git a/Engine/Plugins/Experimental/MutableExtensionExample/Source/MutableExtensionExample/Public/MeshIdentity.h b/Engine/Plugins/Experimental/MutableExtensionExample/Source/MutableExtensionExample/Public/MeshIdentity.h
--- a/Engine/Plugins/Experimental/MutableExtensionExample/Source/MutableExtensionExample/Public/MeshIdentity.h
+++ b/Engine/Plugins/Experimental/MutableExtensionExample/Source/MutableExtensionExample/Public/MeshIdentity.h
@@ -7,16 +7,15 @@
 #include "MeshIdentity.generated.h"
 
 
-USTRUCT()
+USTRUCT(DisplayName = "Mesh Identity")
 struct FMeshIdentity : public UE::Mutable::FExternalOperation
 {
 	GENERATED_BODY()
 	
 	virtual ~FMeshIdentity() override = default;
-	virtual FName GetName() const override;
-	virtual TArray<TPair<FName, const UScriptStruct*>> GetInputs() const override;
-	virtual TPair<FName, const UScriptStruct*> GetOutput() const override;
+	virtual TArray<TPair<FText, const UScriptStruct*>> GetInputs() const override;
+	virtual TPair<FText, const UScriptStruct*> GetOutput() const override;
 	virtual void Evaluate(UE::Mutable::FContext& Context) const override;
 
-	const static FName NameInputMesh;
+	const static FText TextInputMesh;
 };
\ No newline at end of file
diff --git a/Engine/Plugins/Experimental/MutableExtensionExample/Source/MutableExtensionExample/Public/MeshNoise.h b/Engine/Plugins/Experimental/MutableExtensionExample/Source/MutableExtensionExample/Public/MeshNoise.h
--- a/Engine/Plugins/Experimental/MutableExtensionExample/Source/MutableExtensionExample/Public/MeshNoise.h
+++ b/Engine/Plugins/Experimental/MutableExtensionExample/Source/MutableExtensionExample/Public/MeshNoise.h
@@ -7,20 +7,18 @@
 #include "MeshNoise.generated.h"
 
 
-USTRUCT()
+USTRUCT(DisplayName = "Mesh Noise")
 struct FMeshNoise : public UE::Mutable::FExternalOperation
 {
 	GENERATED_BODY()
 	
-	virtual ~FMeshNoise() override = default;
-	virtual FName GetName() const override;
-	virtual TArray<TPair<FName, const UScriptStruct*>> GetInputs() const override;
-	virtual TPair<FName, const UScriptStruct*> GetOutput() const override;
+	virtual TArray<TPair<FText, const UScriptStruct*>> GetInputs() const override;
+	virtual TPair<FText, const UScriptStruct*> GetOutput() const override;
 	virtual void Evaluate(UE::Mutable::FContext& Context) const override;
 
-	const static FName NameInputMesh;
+	const static FText TextInputMesh;
 
-	const static FName NameInputFactor;
+	const static FText TextInputFactor;
 
 	/** Operation constant. This value will be copied when compiling a Customizable Object, hence it can not be changed at runtime.
 	 *
diff --git a/Engine/Plugins/Mutable/Config/DefaultMutable.ini b/Engine/Plugins/Mutable/Config/DefaultMutable.ini
--- a/Engine/Plugins/Mutable/Config/DefaultMutable.ini
+++ b/Engine/Plugins/Mutable/Config/DefaultMutable.ini
@@ -47,6 +47,7 @@
 +ClassRedirects=(OldName="/Script/CustomizableObjectEditor.CustomizableObjectNodeMaterialPinDataVector",NewName="/Script/CustomizableObjectEditor.CONodeSkeletalMeshSectionPinDataVector")
 +ClassRedirects=(OldName="/Script/CustomizableObjectEditor.CustomizableObjectNodeMaterialPinDataScalar",NewName="/Script/CustomizableObjectEditor.CONodeSkeletalMeshSectionPinDataScalar")
 +ClassRedirects=(OldName="/Script/CustomizableObjectEditor.CONodeComponentPassthroughMeshMaterialOverridePinData",NewName="/Script/CustomizableObjectEditor.CONodeComponentPassthroughMeshMaterialPinData")
++ClassRedirects=(OldName="/Script/CustomizableObjectEditor.CustomizableObjectNodeSwitchBase",NewName="/Script/CustomizableObjectEditor.CONodeSwitch")
 
 ; Settings for the Mutable plugin
 [Features]
diff --git a/Engine/Plugins/Mutable/Source/CustomizableObject/Internal/MuCO/CustomizableObjectPrivate.h b/Engine/Plugins/Mutable/Source/CustomizableObject/Internal/MuCO/CustomizableObjectPrivate.h
--- a/Engine/Plugins/Mutable/Source/CustomizableObject/Internal/MuCO/CustomizableObjectPrivate.h
+++ b/Engine/Plugins/Mutable/Source/CustomizableObject/Internal/MuCO/CustomizableObjectPrivate.h
@@ -1452,7 +1452,9 @@ enum class ECustomizableObjectVersions : int32
 	AssetUserDataPassthroughObject2,
 
 	NewProgramAddressFormat,
-
+	
+	IS_SWITCH,
+	
 	// -----<new versions can be added above this line>--------
 	LastCustomizableObjectVersion
 };
diff --git a/Engine/Plugins/Mutable/Source/CustomizableObject/Internal/MuCO/CustomizableObjectResourceDataTypes.h b/Engine/Plugins/Mutable/Source/CustomizableObject/Internal/MuCO/CustomizableObjectResourceDataTypes.h
deleted file mode 100644
diff --git a/Engine/Plugins/Mutable/Source/CustomizableObject/Internal/MuCO/ICustomizableObjectEditorModule.h b/Engine/Plugins/Mutable/Source/CustomizableObject/Internal/MuCO/ICustomizableObjectEditorModule.h
--- a/Engine/Plugins/Mutable/Source/CustomizableObject/Internal/MuCO/ICustomizableObjectEditorModule.h
+++ b/Engine/Plugins/Mutable/Source/CustomizableObject/Internal/MuCO/ICustomizableObjectEditorModule.h
@@ -139,5 +139,9 @@ public:
 	/** Add a CO to the pending - load list.COs need to wait until all related objects are fully loaded before being able
 	 * to do things like check-if-up-to-date or compile. */
 	virtual void AddPendingLoad(UCustomizableObject* Object) = 0;
+
+	virtual void RegisterExternalOperation(const UScriptStruct* Operation) = 0;
+
+	virtual void UnregisterExternalOperation(const UScriptStruct* Operation) = 0;
 };
  
diff --git a/Engine/Plugins/Mutable/Source/CustomizableObject/Internal/MuCO/ICustomizableObjectModulePrivate.h b/Engine/Plugins/Mutable/Source/CustomizableObject/Internal/MuCO/ICustomizableObjectModulePrivate.h
--- a/Engine/Plugins/Mutable/Source/CustomizableObject/Internal/MuCO/ICustomizableObjectModulePrivate.h
+++ b/Engine/Plugins/Mutable/Source/CustomizableObject/Internal/MuCO/ICustomizableObjectModulePrivate.h
@@ -20,9 +20,6 @@ public:
 		return FModuleManager::LoadModuleChecked<ICustomizableObjectModulePrivate>("CustomizableObject");
 	}
 	
-	virtual TArray<TStrongObjectPtr<const UScriptStruct>> GetRegisteredExternalOperations() const = 0;
+	virtual TSet<TStrongObjectPtr<const UScriptStruct>> GetRegisteredExternalOperations() const = 0;
 };
 
-
-CUSTOMIZABLEOBJECT_API bool ContainsUObject(const UE::Mutable::FExternalOperation& Operation, const UScriptStruct*& OutExternalType);
-
diff --git a/Engine/Plugins/Mutable/Source/CustomizableObject/Private/MuCO/CustomizableObject.cpp b/Engine/Plugins/Mutable/Source/CustomizableObject/Private/MuCO/CustomizableObject.cpp
--- a/Engine/Plugins/Mutable/Source/CustomizableObject/Private/MuCO/CustomizableObject.cpp
+++ b/Engine/Plugins/Mutable/Source/CustomizableObject/Private/MuCO/CustomizableObject.cpp
@@ -51,7 +51,6 @@
 
 #include "ExternalOperationProvider.h"
 #include "MuCO/CustomizableObjectCustomVersion.h"
-#include "MuCO/CustomizableObjectResourceDataTypes.h"
 #include "MuCO/CustomizableObjectSkeletalMesh.h"
 
 #include UE_INLINE_GENERATED_CPP_BY_NAME(CustomizableObject)
diff --git a/Engine/Plugins/Mutable/Source/CustomizableObject/Private/MuCO/CustomizableObjectInstance.cpp b/Engine/Plugins/Mutable/Source/CustomizableObject/Private/MuCO/CustomizableObjectInstance.cpp
--- a/Engine/Plugins/Mutable/Source/CustomizableObject/Private/MuCO/CustomizableObjectInstance.cpp
+++ b/Engine/Plugins/Mutable/Source/CustomizableObject/Private/MuCO/CustomizableObjectInstance.cpp
@@ -44,7 +44,6 @@
 #include "PhysicsEngine/PhysicsConstraintTemplate.h"
 #include "PhysicsEngine/SkeletalBodySetup.h"
 #include "MuCO/CustomizableObjectCustomVersion.h"
-#include "MuCO/CustomizableObjectResourceDataTypes.h"
 #include "MuCO/LoadUtils.h"
 #include "Serialization/BulkData.h"
 
diff --git a/Engine/Plugins/Mutable/Source/CustomizableObject/Private/MuCO/CustomizableObjectModule.cpp b/Engine/Plugins/Mutable/Source/CustomizableObject/Private/MuCO/CustomizableObjectModule.cpp
--- a/Engine/Plugins/Mutable/Source/CustomizableObject/Private/MuCO/CustomizableObjectModule.cpp
+++ b/Engine/Plugins/Mutable/Source/CustomizableObject/Private/MuCO/CustomizableObjectModule.cpp
@@ -20,6 +20,7 @@
 #include "UObject/UObjectIterator.h"
 #include "GPUSkinPublicDefs.h"
 #include "Components/SkeletalMeshComponent.h"
+#include "MuCO/ICustomizableObjectEditorModule.h"
 #include "MuR/External/Operation.h"
 
 /**
@@ -46,7 +47,7 @@ public:
 	virtual void UnregisterExternalOperation(const UScriptStruct* Operation) override;
 
 	// ICustomizableObjectModulePrivate interface
-	virtual TArray<TStrongObjectPtr<const UScriptStruct>> GetRegisteredExternalOperations() const override;
+	virtual TSet<TStrongObjectPtr<const UScriptStruct>> GetRegisteredExternalOperations() const override;
 
 private:
 	void RefreshExtensionData();
@@ -65,7 +66,7 @@ private:
 	TArray<FRegisteredCustomizableObjectPinType> ExtendedPinTypes;
 	TArray<FRegisteredObjectNodeInputPin> AdditionalObjectNodePins;
 
-	TArray<TStrongObjectPtr<const UScriptStruct>> RegisteredOperations;
+	TSet<TStrongObjectPtr<const UScriptStruct>> RegisteredOperations;
 };
 
 
@@ -216,31 +217,6 @@ bool ContainsUObject(const UScriptStruct* ScriptStruct)
 }
 
 
-bool ContainsUObject(const UE::Mutable::FExternalOperation& Operation, const UScriptStruct*& OutExternalType)
-{
-	TArray<const UScriptStruct*> Types;
-	
-	for (const TPair<FName, const UScriptStruct*>& Input : Operation.GetInputs())
-	{
-		Types.AddUnique(Input.Value);
-	}
-
-	Types.AddUnique((Operation.GetOutput().Value));
-
-	for (const UScriptStruct* Type : Types)
-	{
-		if (ContainsUObject(Type))
-		{
-			OutExternalType = Type;
-			return true;
-		}
-	}
-
-	OutExternalType = nullptr;
-	return false;
-}
-
-
 void FCustomizableObjectModule::RegisterExternalOperation(const UScriptStruct* Operation)
 {
 	if (!Operation)
@@ -254,17 +230,12 @@ void FCustomizableObjectModule::RegisterExternalOperation(const UScriptStruct* O
 		return;
 	}
 
-	FInstancedStruct InstancedStruct;
-	InstancedStruct.InitializeAs(Operation);
+	RegisteredOperations.Add(TStrongObjectPtr(Operation));
 	
-	const UScriptStruct* ExternalType;
-	if (ContainsUObject(InstancedStruct.Get<UE::Mutable::FExternalOperation>(), ExternalType))
+	if (ICustomizableObjectEditorModule* EditorModule = ICustomizableObjectEditorModule::Get())
 	{
-		UE_LOG(LogMutable, Error, TEXT("Unable to register operation. Operation contains UObject properties. Operation: %s Type: %s"), *Operation->GetName(), *ExternalType->GetName());
-		return;
+		EditorModule->RegisterExternalOperation(Operation);
 	}
-	
-	RegisteredOperations.AddUnique(TStrongObjectPtr(Operation));
 }
 
 
@@ -275,18 +246,16 @@ void FCustomizableObjectModule::UnregisterExternalOperation(const UScriptStruct*
 		return;
 	}
 
-	for (TArray<TStrongObjectPtr<const UScriptStruct>>::TIterator It = RegisteredOperations.CreateIterator(); It; ++It)
+	RegisteredOperations.Remove(TStrongObjectPtr(Operation));
+	
+	if (ICustomizableObjectEditorModule* EditorModule = ICustomizableObjectEditorModule::Get())
 	{
-		if (It->Get() == Operation)
-		{
-			It.RemoveCurrent();
-			break;
-		}
+		EditorModule->UnregisterExternalOperation(Operation);
 	}
 }
 
 
-TArray<TStrongObjectPtr<const UScriptStruct>> FCustomizableObjectModule::GetRegisteredExternalOperations() const
+TSet<TStrongObjectPtr<const UScriptStruct>> FCustomizableObjectModule::GetRegisteredExternalOperations() const
 {
 	return RegisteredOperations;
 }
diff --git a/Engine/Plugins/Mutable/Source/CustomizableObjectEditor/Internal/ICustomizableObjectEditorModulePrivate.h b/Engine/Plugins/Mutable/Source/CustomizableObjectEditor/Internal/ICustomizableObjectEditorModulePrivate.h
--- a/Engine/Plugins/Mutable/Source/CustomizableObjectEditor/Internal/ICustomizableObjectEditorModulePrivate.h
+++ b/Engine/Plugins/Mutable/Source/CustomizableObjectEditor/Internal/ICustomizableObjectEditorModulePrivate.h
@@ -23,4 +23,10 @@ public:
 	}
 	
 	virtual void EnqueueCompileRequest(const TSharedRef<FCompilationRequest>& InCompilationRequest, bool bForceRequest = false) = 0;
+	
+	virtual const TMap<FName, FLinearColor>& GetExternalPinTypeColors() const = 0;
+
+	virtual const TMap<FName, FText>& GetExternalPinFriendlyName() const = 0;
+	
+	virtual const TMap<FName, TStrongObjectPtr<const UScriptStruct>>& GetExternalPinTypes() const = 0;
 };
\ No newline at end of file
diff --git a/Engine/Plugins/Mutable/Source/CustomizableObjectEditor/Private/MuCOE/CustomizableObjectEditorModule.cpp b/Engine/Plugins/Mutable/Source/CustomizableObjectEditor/Private/MuCOE/CustomizableObjectEditorModule.cpp
--- a/Engine/Plugins/Mutable/Source/CustomizableObjectEditor/Private/MuCOE/CustomizableObjectEditorModule.cpp
+++ b/Engine/Plugins/Mutable/Source/CustomizableObjectEditor/Private/MuCOE/CustomizableObjectEditorModule.cpp
@@ -1161,6 +1161,147 @@ void FCustomizableObjectEditorModule::EnqueueCompileRequest(const TSharedRef<FCo
 }
 
 
+const TMap<FName, FLinearColor>& FCustomizableObjectEditorModule::GetExternalPinTypeColors() const
+{
+	return ExternalTypeTypePinColors;
+}
+
+
+void FCustomizableObjectEditorModule::RegisterExternalOperation(const UScriptStruct* Operation)
+{
+	if (!Operation)
+	{
+		return;
+	}
+	
+	if (!Operation->IsChildOf(UE::Mutable::FExternalOperation::StaticStruct()))
+	{
+		return;
+	}
+	
+	auto Register = [&](const UScriptStruct* ExternalPinType)
+	{
+		if (ExternalPinType->GetCppStructOps()->GetCapabilities().IsAbstract)
+		{
+			return;
+		}
+
+		FName PinType = UEdGraphSchema_CustomizableObject::GetPinCategoryName(ExternalPinType);
+		FLinearColor Color;
+			
+		const FString Hex = ExternalPinType->GetMetaData("PinColor"); // e.g. "#FF00AAFF"
+		if (!Hex.IsEmpty())
+		{
+			Color = FColor::FromHex(Hex);
+		}
+		else
+		{
+			const uint32 Hash = GetTypeHash(PinType.ToString());
+			Color = FLinearColor::MakeRandomSeededColor(Hash);
+		}
+
+		TStrongObjectPtr<const UScriptStruct> StrongExternalPinType = TStrongObjectPtr(ExternalPinType);
+	
+		TStrongObjectPtr<const UScriptStruct>* Result = ExternalTypeCategories.Find(PinType);
+		if (Result && *Result != StrongExternalPinType)
+		{
+			UE_LOGF(LogMutable, Error, "External Pin Type already registered: %ls", *PinType.ToString());
+			return;
+		}
+			
+		ExternalTypeTypePinColors.Add(PinType, Color);
+		ExternalTypeFriendlyName.Add(PinType, ExternalPinType->GetDisplayNameText());
+		ExternalTypeCategories.Add(PinType, StrongExternalPinType);
+			
+		++RegisteredExternalType.FindOrAdd(PinType, 0);
+			
+		check(RegisteredExternalType.Num() == ExternalTypeTypePinColors.Num());
+		check(RegisteredExternalType.Num() == ExternalTypeFriendlyName.Num());
+		check(RegisteredExternalType.Num() == ExternalTypeCategories.Num());
+	};
+	
+	FInstancedStruct InstancedStruct;
+	InstancedStruct.InitializeAs(Operation);
+	
+	const UE::Mutable::FExternalOperation& ExternalOperation = InstancedStruct.Get<UE::Mutable::FExternalOperation>();
+
+	Register(ExternalOperation.GetOutput().Value);
+	
+	for (const TPair<FText, const UScriptStruct*>& Input : ExternalOperation.GetInputs())
+	{
+		Register(Input.Value);
+	}
+}
+
+
+void FCustomizableObjectEditorModule::UnregisterExternalOperation(const UScriptStruct* Operation)
+{
+	if (!Operation)
+	{
+		return;
+	}
+	
+	if (!Operation->IsChildOf(UE::Mutable::FExternalOperation::StaticStruct()))
+	{
+		return;
+	}
+
+	auto Unregister = [&](const UScriptStruct* ExternalPinType)
+	{
+		if (ExternalPinType->GetCppStructOps()->GetCapabilities().IsAbstract)
+		{
+			return;
+		}
+			
+		FName PinType = UEdGraphSchema_CustomizableObject::GetPinCategoryName(ExternalPinType);
+
+		uint32* Count = RegisteredExternalType.Find(PinType);
+		if (!Count)
+		{
+			return;
+		}
+			
+		--*Count;
+		if (*Count == 0)
+		{
+			ExternalTypeTypePinColors.Remove(PinType);
+			ExternalTypeFriendlyName.Remove(PinType);
+			ExternalTypeCategories.Remove(PinType);
+			
+			RegisteredExternalType.Remove(PinType);
+			
+			check(RegisteredExternalType.Num() == ExternalTypeTypePinColors.Num());
+			check(RegisteredExternalType.Num() == ExternalTypeFriendlyName.Num());
+			check(RegisteredExternalType.Num() == ExternalTypeCategories.Num());
+		}
+	};
+	
+	FInstancedStruct InstancedStruct;
+	InstancedStruct.InitializeAs(Operation);
+	
+	const UE::Mutable::FExternalOperation& ExternalOperation = InstancedStruct.Get<UE::Mutable::FExternalOperation>();
+
+	Unregister(ExternalOperation.GetOutput().Value);
+	
+	for (const TPair<FText, const UScriptStruct*>& Input : ExternalOperation.GetInputs())
+	{
+		Unregister(Input.Value);
+	}
+}
+
+
+const TMap<FName, TStrongObjectPtr<const UScriptStruct>>& FCustomizableObjectEditorModule::GetExternalPinTypes() const
+{
+	return ExternalTypeCategories;
+}
+
+
+const TMap<FName, FText>& FCustomizableObjectEditorModule::GetExternalPinFriendlyName() const
+{
+	return ExternalTypeFriendlyName;
+}
+
+
 void FCustomizableObjectEditorModule::CompileCustomizableObject(UCustomizableObject& Object, const FCompileParams* Params, bool bSilent, bool bForce)
 {
 	UCustomizableObjectSystem* System = UCustomizableObjectSystem::GetInstance();
diff --git a/Engine/Plugins/Mutable/Source/CustomizableObjectEditor/Private/MuCOE/CustomizableObjectEditorModule.h b/Engine/Plugins/Mutable/Source/CustomizableObjectEditor/Private/MuCOE/CustomizableObjectEditorModule.h
--- a/Engine/Plugins/Mutable/Source/CustomizableObjectEditor/Private/MuCOE/CustomizableObjectEditorModule.h
+++ b/Engine/Plugins/Mutable/Source/CustomizableObjectEditor/Private/MuCOE/CustomizableObjectEditorModule.h
@@ -78,7 +78,12 @@ public:
 
 	// ICustomizableObjectEditorModulePrivate interface
 	virtual void EnqueueCompileRequest(const TSharedRef<FCompilationRequest>& InCompilationRequest, bool bForceRequest = false) override;
-
+	virtual const TMap<FName, FLinearColor>& GetExternalPinTypeColors() const override;
+	virtual void RegisterExternalOperation(const UScriptStruct* Operation) override;
+	virtual void UnregisterExternalOperation(const UScriptStruct* Operation) override;
+	virtual const TMap<FName, TStrongObjectPtr<const UScriptStruct>>& GetExternalPinTypes() const override;
+	virtual const TMap<FName, FText>& GetExternalPinFriendlyName() const override;
+	
 	// FTickableGameObject interface
 	virtual ETickableTickType GetTickableTickType() const override { return ETickableTickType::Conditional; };
 	virtual bool IsTickable() const override;
@@ -95,6 +100,15 @@ private:
 	/** List of registered custom details to remove later. */
 	TArray<FName> RegisteredCustomDetails;
 
+	/** How many times a pin type has been registered. */
+	TMap<FName, uint32> RegisteredExternalType;
+	
+	TMap<FName, FLinearColor> ExternalTypeTypePinColors;
+	
+	TMap<FName, FText> ExternalTypeFriendlyName;
+	
+	TMap<FName, TStrongObjectPtr<const UScriptStruct>> ExternalTypeCategories;
+
 	virtual void OnUpstreamCOsLoaded(UCustomizableObject* Object) const;
 	
 	/** Register Custom details. Also adds them to RegisteredCustomDetails list. */
diff --git a/Engine/Plugins/Mutable/Source/CustomizableObjectEditor/Private/MuCOE/CustomizableObjectSchemaActions.h b/Engine/Plugins/Mutable/Source/CustomizableObjectEditor/Private/MuCOE/CustomizableObjectSchemaActions.h
--- a/Engine/Plugins/Mutable/Source/CustomizableObjectEditor/Private/MuCOE/CustomizableObjectSchemaActions.h
+++ b/Engine/Plugins/Mutable/Source/CustomizableObjectEditor/Private/MuCOE/CustomizableObjectSchemaActions.h
@@ -26,8 +26,8 @@ struct FCustomizableObjectSchemaAction_NewNode : public FEdGraphSchemaAction
 		, NodeTemplate(NULL)
 	{}
 
-	FCustomizableObjectSchemaAction_NewNode(const FString& InNodeCategory, const FText& InMenuDesc, const FText& InToolTip, const int32 InGrouping, const FText& InKeywords = FText(), int32 InSectionID = 0)
-		: FEdGraphSchemaAction(FText::FromString(InNodeCategory), InMenuDesc, InToolTip, InGrouping, InKeywords, InSectionID)
+	FCustomizableObjectSchemaAction_NewNode(const FText& InNodeCategory, const FText& InMenuDesc, const FText& InToolTip, const int32 InGrouping, const FText& InKeywords = FText(), int32 InSectionID = 0)
+		: FEdGraphSchemaAction(InNodeCategory, InMenuDesc, InToolTip, InGrouping, InKeywords, InSectionID)
 		, NodeTemplate(NULL)
 	{}
 
diff --git a/Engine/Plugins/Mutable/Source/CustomizableObjectEditor/Private/MuCOE/EdGraphSchema_CustomizableObject.cpp b/Engine/Plugins/Mutable/Source/CustomizableObjectEditor/Private/MuCOE/EdGraphSchema_CustomizableObject.cpp
--- a/Engine/Plugins/Mutable/Source/CustomizableObjectEditor/Private/MuCOE/EdGraphSchema_CustomizableObject.cpp
+++ b/Engine/Plugins/Mutable/Source/CustomizableObjectEditor/Private/MuCOE/EdGraphSchema_CustomizableObject.cpp
@@ -16,6 +16,7 @@
 #include "Framework/Application/SlateApplication.h"
 #include "Framework/Commands/GenericCommands.h"
 #include "GraphEditorActions.h"
+#include "ICustomizableObjectEditorModulePrivate.h"
 #include "Materials/MaterialInterface.h"
 #include "MuCO/CustomizableObject.h"
 #include "MuCO/CustomizableObjectExtension.h"
@@ -27,14 +28,12 @@
 #include "MuCOE/Nodes/CONodeMaterialConstant.h"
 #include "MuCOE/Nodes/CONodeMaterialModify.h"
 #include "MuCOE/Nodes/CONodeMaterialVariation.h"
-#include "MuCOE/Nodes/CONodeMaterialSwitch.h"
 #include "MuCOE/Nodes/CONodeModifierTransformWithBone.h"
 #include "MuCOE/Nodes/CustomizableObjectNodeAnimationPose.h"
 #include "MuCOE/Nodes/CustomizableObjectNodeColorArithmeticOp.h"
 #include "MuCOE/Nodes/CustomizableObjectNodeColorConstant.h"
 #include "MuCOE/Nodes/CustomizableObjectNodeColorFromFloats.h"
 #include "MuCOE/Nodes/CustomizableObjectNodeColorParameter.h"
-#include "MuCOE/Nodes/CustomizableObjectNodeColorSwitch.h"
 #include "MuCOE/Nodes/CustomizableObjectNodeColorVariation.h"
 #include "MuCOE/Nodes/CustomizableObjectNodeColorToSRGB.h"
 #include "MuCOE/Nodes/CONodeCopySkeletalMeshSection.h"
@@ -44,13 +43,11 @@
 #include "MuCOE/Nodes/CustomizableObjectNodeExternalPin.h"
 #include "MuCOE/Nodes/CustomizableObjectNodeFloatConstant.h"
 #include "MuCOE/Nodes/CustomizableObjectNodeFloatParameter.h"
-#include "MuCOE/Nodes/CustomizableObjectNodeFloatSwitch.h"
 #include "MuCOE/Nodes/CustomizableObjectNodeFloatVariation.h"
 #include "MuCOE/Nodes/CustomizableObjectNodeFloatArithmeticOp.h"
 #include "MuCOE/Nodes/CustomizableObjectNodeGroupProjectorParameter.h"
 #include "MuCOE/Nodes/CustomizableObjectNodeMacroInstance.h"
 #include "MuCOE/Nodes/CustomizableObjectNodeMaterialParameter.h"
-#include "MuCOE/Nodes/CONodeSkeletalMeshSectionSwitch.h"
 #include "MuCOE/Nodes/CONodeSkeletalMeshSectionVariation.h"
 #include "MuCOE/Nodes/CONodeModifierExtendSkeletalMeshSection.h"
 #include "MuCOE/Nodes/CONodeModifierEditSkeletalMeshSection.h"
@@ -66,14 +63,12 @@
 #include "MuCOE/Nodes/CustomizableObjectNodeMeshMorphStackApplication.h"
 #include "MuCOE/Nodes/CustomizableObjectNodeMeshMorphStackDefinition.h"
 #include "MuCOE/Nodes/CustomizableObjectNodeMeshReshape.h"
-#include "MuCOE/Nodes/CustomizableObjectNodeMeshSwitch.h"
 #include "MuCOE/Nodes/CustomizableObjectNodeMeshVariation.h"
 #include "MuCOE/Nodes/CustomizableObjectNodeSkeletalMeshParameter.h"
 #include "MuCOE/Nodes/CustomizableObjectNodeObjectChild.h"
 #include "MuCOE/Nodes/CustomizableObjectNodeObjectGroup.h"
 #include "MuCOE/Nodes/CustomizableObjectNodeComponentPassthroughMesh.h"
 #include "MuCOE/Nodes/CustomizableObjectNodeComponentVariation.h"
-#include "MuCOE/Nodes/CustomizableObjectNodeComponentSwitch.h"
 #include "MuCOE/Nodes/CustomizableObjectNodeProjectorConstant.h"
 #include "MuCOE/Nodes/CustomizableObjectNodeSkeletalMesh.h"
 #include "MuCOE/Nodes/CustomizableObjectNodeStaticMesh.h"
@@ -92,8 +87,6 @@
 #include "MuCOE/Nodes/CustomizableObjectNodeTextureParameter.h"
 #include "MuCOE/Nodes/CustomizableObjectNodeTextureProject.h"
 #include "MuCOE/Nodes/CustomizableObjectNodeTextureSample.h"
-#include "MuCOE/Nodes/CustomizableObjectNodeTextureSwitch.h"
-#include "MuCOE/Nodes/CustomizableObjectNodePassThroughTextureSwitch.h"
 #include "MuCOE/Nodes/CustomizableObjectNodeTextureToChannels.h"
 #include "MuCOE/Nodes/CustomizableObjectNodeTextureTransform.h"
 #include "MuCOE/Nodes/CustomizableObjectNodeTextureSaturate.h"
@@ -125,6 +118,7 @@
 #include "Nodes/CONodeSkeletalMeshBreak.h"
 #include "Nodes/CONodeTransformWithBone.h"
 #include "Nodes/CONodeTransformInMesh.h"
+#include "Nodes/CONodeSwitch.h"
 #include "UObject/Class.h"
 #include "Widgets/SMutablePinTypeSelector.h"
 
@@ -335,7 +329,7 @@ const FText UEdGraphSchema_CustomizableObject::NC_Experimental(FText::FromString
 const FText UEdGraphSchema_CustomizableObject::NC_Material(FText::FromString("Material"));
 
 
-TSharedPtr<FCustomizableObjectSchemaAction_NewNode> AddNewNodeAction(FGraphActionListBuilderBase& ContextMenuBuilder, const FString& Category, const FText& MenuDesc, const FText& Tooltip, const int32 Grouping = 0, const FString& Keywords = FString())
+TSharedPtr<FCustomizableObjectSchemaAction_NewNode> AddNewNodeAction(FGraphActionListBuilderBase& ContextMenuBuilder, const FText& Category, const FText& MenuDesc, const FText& Tooltip, const int32 Grouping = 0, const FString& Keywords = FString())
 {
 	TSharedPtr<FCustomizableObjectSchemaAction_NewNode> NewActionNode = MakeShared<FCustomizableObjectSchemaAction_NewNode>(Category, MenuDesc, Tooltip, Grouping, FText::FromString(Keywords));
 
@@ -345,7 +339,7 @@ TSharedPtr<FCustomizableObjectSchemaAction_NewNode> AddNewNodeAction(FGraphActio
 }
 
 
-TSharedPtr<FCustomizableObjectSchemaAction_NewNode> AddNewNodeAction(TArray< TSharedPtr<FEdGraphSchemaAction> >& OutTypes, const FString& Category, const FText& MenuDesc, const FText& Tooltip)
+TSharedPtr<FCustomizableObjectSchemaAction_NewNode> AddNewNodeAction(TArray< TSharedPtr<FEdGraphSchemaAction> >& OutTypes, const FText& Category, const FText& MenuDesc, const FText& Tooltip)
 {
 	return *(new (OutTypes) TSharedPtr<FCustomizableObjectSchemaAction_NewNode>(new FCustomizableObjectSchemaAction_NewNode(Category, MenuDesc, Tooltip, 0)));
 }
@@ -392,7 +386,7 @@ namespace
 	}
 
 	template<class FilterFn>
-	void AddNewNodeActionFiltered(UCustomizableObjectNode* TemplateNode, FGraphContextMenuBuilder& ContextMenuBuilder, const FString& Category, const FText& MenuDesc, const int32 Grouping, FilterFn&& Filter)
+	void AddNewNodeActionFiltered(UCustomizableObjectNode* TemplateNode, FGraphContextMenuBuilder& ContextMenuBuilder, const FText& Category, const FText& MenuDesc, const int32 Grouping, FilterFn&& Filter)
 	{
 		if (!Filter(TemplateNode, ContextMenuBuilder))
 		{
@@ -404,13 +398,13 @@ namespace
 	}
 
 	template<class FilterFn>
-	void AddNewNodeActionFiltered(UCustomizableObjectNode* TemplateNode, FGraphContextMenuBuilder& ContextMenuBuilder, const FString& Category, const int32 Grouping, FilterFn&& Filter)
+	void AddNewNodeActionFiltered(UCustomizableObjectNode* TemplateNode, FGraphContextMenuBuilder& ContextMenuBuilder, const FText& Category, const int32 Grouping, FilterFn&& Filter)
 	{
 		AddNewNodeActionFiltered(TemplateNode, ContextMenuBuilder, Category, TemplateNode->GetNodeTitle(ENodeTitleType::ListView), Grouping, Forward<FilterFn>(Filter));
 	}
 
 	template<size_t N, class FilterFn>
-	void AddNewNodeCategoryActionsFiltered(UCustomizableObjectNode* (&TemplateNodes)[N], FGraphContextMenuBuilder& ContextMenuBuilder, const FString& Category, const int32 Grouping, FilterFn&& Filter)
+	void AddNewNodeCategoryActionsFiltered(UCustomizableObjectNode* (&TemplateNodes)[N], FGraphContextMenuBuilder& ContextMenuBuilder, const FText& Category, const int32 Grouping, FilterFn&& Filter)
 	{
 		for (size_t i = 0; i < N; ++i)
 		{
@@ -451,7 +445,7 @@ void UEdGraphSchema_CustomizableObject::GetGraphContextActions(FGraphContextMenu
 
 	{
 		UCustomizableObjectNode* Node = ContextMenuBuilder.CreateTemplateNode<UCustomizableObjectNodeObject>();
-		AddNewNodeActionFiltered(Node, ContextMenuBuilder, TEXT("Object"), LOCTEXT("Base_Group", "Base Object"), GeneralGrouping, 
+		AddNewNodeActionFiltered(Node, ContextMenuBuilder, GetPinCategoryFriendlyName(PC_Object), LOCTEXT("Base_Group", "Base Object"), GeneralGrouping, 
 			[&Filter](UCustomizableObjectNode* TemplateNode, FGraphContextMenuBuilder& ContextMenuBuilder) -> bool
 			{
 				// Only let user add a base node if there isn't one in the graph
@@ -472,12 +466,12 @@ void UEdGraphSchema_CustomizableObject::GetGraphContextActions(FGraphContextMenu
 
 	{
 		UCustomizableObjectNode* Node = NewObject<UCustomizableObjectNodeObjectGroup>(); //ContextMenuBuilder.CreateTemplateNode<UCustomizableObjectNodeObjectGroup>();
-		AddNewNodeActionFiltered(Node, ContextMenuBuilder, TEXT("Object"), LOCTEXT("Child_Group", "Object Group"), GeneralGrouping, Filter);
+		AddNewNodeActionFiltered(Node, ContextMenuBuilder, GetPinCategoryFriendlyName(PC_Object), LOCTEXT("Child_Group", "Object Group"), GeneralGrouping, Filter);
 	}
 	
 	{
 		UCustomizableObjectNode* Node = ContextMenuBuilder.CreateTemplateNode<UCustomizableObjectNodeObjectChild>();
-		AddNewNodeActionFiltered(Node, ContextMenuBuilder, TEXT("Object"), LOCTEXT("Child_Object", "Child Object"), GeneralGrouping, Filter);
+		AddNewNodeActionFiltered(Node, ContextMenuBuilder, GetPinCategoryFriendlyName(PC_Object), LOCTEXT("Child_Object", "Child Object"), GeneralGrouping, Filter);
 	}
 
 	{
@@ -485,14 +479,13 @@ void UEdGraphSchema_CustomizableObject::GetGraphContextActions(FGraphContextMenu
 		{
 			ContextMenuBuilder.CreateTemplateNode<UCONodeSkeletalMeshSection>(),
 			ContextMenuBuilder.CreateTemplateNode<UCONodeSkeletalMeshSectionVariation>(),
-			ContextMenuBuilder.CreateTemplateNode<UCONodeSkeletalMeshSectionSwitch>(),
 			ContextMenuBuilder.CreateTemplateNode<UCONodeCopySkeletalMeshSection>(),
 			ContextMenuBuilder.CreateTemplateNode<UCustomizableObjectNodeTable>(),
 			ContextMenuBuilder.CreateTemplateNode<UCustomizableObjectNodeMacroInstance>(),
 			ContextMenuBuilder.CreateTemplateNode<UCustomizableObjectNodeStaticString>()
 		};
 
-		AddNewNodeCategoryActionsFiltered(ObjectTemplateNodes, ContextMenuBuilder, TEXT("Object"), GeneralGrouping, Filter);
+		AddNewNodeCategoryActionsFiltered(ObjectTemplateNodes, ContextMenuBuilder, GetPinCategoryFriendlyName(PC_Object), GeneralGrouping, Filter);
 	}
 
 	{
@@ -510,7 +503,7 @@ void UEdGraphSchema_CustomizableObject::GetGraphContextActions(FGraphContextMenu
 			ContextMenuBuilder.CreateTemplateNode<UCustomizableObjectNodeModifierTransformInMesh>(),
 		};
 
-		AddNewNodeCategoryActionsFiltered(ModifierTemplateNodes, ContextMenuBuilder, TEXT("Modifier"), GeneralGrouping, Filter);
+		AddNewNodeCategoryActionsFiltered(ModifierTemplateNodes, ContextMenuBuilder, GetPinCategoryFriendlyName(PC_Modifier), GeneralGrouping, Filter);
 	}
 
 	{
@@ -519,10 +512,9 @@ void UEdGraphSchema_CustomizableObject::GetGraphContextActions(FGraphContextMenu
 			ContextMenuBuilder.CreateTemplateNode<UCustomizableObjectNodeComponentMesh>(),
 			ContextMenuBuilder.CreateTemplateNode<UCustomizableObjectNodeComponentMeshAddTo>(),
 			ContextMenuBuilder.CreateTemplateNode<UCustomizableObjectNodeComponentVariation>(),
-			ContextMenuBuilder.CreateTemplateNode<UCustomizableObjectNodeComponentSwitch>(),
 		};
 
-		AddNewNodeCategoryActionsFiltered(ComponentTemplateNodes, ContextMenuBuilder, TEXT("Component"), GeneralGrouping, Filter);
+		AddNewNodeCategoryActionsFiltered(ComponentTemplateNodes, ContextMenuBuilder, GetPinCategoryFriendlyName(PC_Component), GeneralGrouping, Filter);
 	}
 
 	{
@@ -531,7 +523,6 @@ void UEdGraphSchema_CustomizableObject::GetGraphContextActions(FGraphContextMenu
 			ContextMenuBuilder.CreateTemplateNode<UCustomizableObjectNodeSkeletalMesh>(),
 			ContextMenuBuilder.CreateTemplateNode<UCustomizableObjectNodeStaticMesh>(),
 			ContextMenuBuilder.CreateTemplateNode<UCustomizableObjectNodeMeshMorph>(),
-			ContextMenuBuilder.CreateTemplateNode<UCustomizableObjectNodeMeshSwitch>(),
 			ContextMenuBuilder.CreateTemplateNode<UCustomizableObjectNodeMeshVariation>(),
 			ContextMenuBuilder.CreateTemplateNode<UCustomizableObjectNodeMeshMorphStackDefinition>(),
 			ContextMenuBuilder.CreateTemplateNode<UCustomizableObjectNodeMeshMorphStackApplication>(),
@@ -542,7 +533,7 @@ void UEdGraphSchema_CustomizableObject::GetGraphContextActions(FGraphContextMenu
 			ContextMenuBuilder.CreateTemplateNode<UCONodeTransformInMesh>(),
 		};
 
-		AddNewNodeCategoryActionsFiltered(MeshTemplateNodes, ContextMenuBuilder, TEXT("Mesh"), GeneralGrouping, Filter);
+		AddNewNodeCategoryActionsFiltered(MeshTemplateNodes, ContextMenuBuilder, GetPinCategoryFriendlyName(PC_Mesh), GeneralGrouping, Filter);
 	}
 	
 	{
@@ -553,8 +544,6 @@ void UEdGraphSchema_CustomizableObject::GetGraphContextActions(FGraphContextMenu
 			ContextMenuBuilder.CreateTemplateNode<UCustomizableObjectNodeTextureBinarise>(),
 			ContextMenuBuilder.CreateTemplateNode<UCustomizableObjectNodeTextureInterpolate>(),
 			ContextMenuBuilder.CreateTemplateNode<UCustomizableObjectNodeTextureLayer>(),
-			ContextMenuBuilder.CreateTemplateNode<UCustomizableObjectNodeTextureSwitch>(),
-			ContextMenuBuilder.CreateTemplateNode<UCustomizableObjectNodePassThroughTextureSwitch>(),
 			ContextMenuBuilder.CreateTemplateNode<UCustomizableObjectNodeTextureVariation>(),
 			ContextMenuBuilder.CreateTemplateNode<UCustomizableObjectNodePassThroughTextureVariation>(),
 			ContextMenuBuilder.CreateTemplateNode<UCustomizableObjectNodeTextureToChannels>(),
@@ -568,7 +557,7 @@ void UEdGraphSchema_CustomizableObject::GetGraphContextActions(FGraphContextMenu
 			ContextMenuBuilder.CreateTemplateNode<UCustomizableObjectNodeTextureSaturate>(),
 		};
 
-		AddNewNodeCategoryActionsFiltered(TextureTemplateNodes, ContextMenuBuilder, TEXT("Texture"), GeneralGrouping, Filter);
+		AddNewNodeCategoryActionsFiltered(TextureTemplateNodes, ContextMenuBuilder, GetPinCategoryFriendlyName(PC_Texture), GeneralGrouping, Filter);
 	}
 
 	{
@@ -576,7 +565,6 @@ void UEdGraphSchema_CustomizableObject::GetGraphContextActions(FGraphContextMenu
 		{
 			ContextMenuBuilder.CreateTemplateNode<UCustomizableObjectNodeColorConstant>(),
 			ContextMenuBuilder.CreateTemplateNode<UCustomizableObjectNodeColorParameter>(),
-			ContextMenuBuilder.CreateTemplateNode<UCustomizableObjectNodeColorSwitch>(),
 			ContextMenuBuilder.CreateTemplateNode<UCustomizableObjectNodeColorVariation>(),
 			ContextMenuBuilder.CreateTemplateNode<UCustomizableObjectNodeTextureSample>(),
 			ContextMenuBuilder.CreateTemplateNode<UCustomizableObjectNodeColorArithmeticOp>(),
@@ -584,7 +572,7 @@ void UEdGraphSchema_CustomizableObject::GetGraphContextActions(FGraphContextMenu
 			ContextMenuBuilder.CreateTemplateNode<UCustomizableObjectNodeColorToSRGB>(),
 		};
 
-		AddNewNodeCategoryActionsFiltered(ColorTemplateNodes, ContextMenuBuilder, TEXT("Color"), GeneralGrouping, Filter);
+		AddNewNodeCategoryActionsFiltered(ColorTemplateNodes, ContextMenuBuilder, GetPinCategoryFriendlyName(PC_Color), GeneralGrouping, Filter);
 	}
 
 	{
@@ -593,7 +581,7 @@ void UEdGraphSchema_CustomizableObject::GetGraphContextActions(FGraphContextMenu
 			ContextMenuBuilder.CreateTemplateNode<UCustomizableObjectNodeEnumParameter>(),
 		};
 
-		AddNewNodeCategoryActionsFiltered(EnumTemplateNodes, ContextMenuBuilder, TEXT("Enum"), GeneralGrouping, Filter);
+		AddNewNodeCategoryActionsFiltered(EnumTemplateNodes, ContextMenuBuilder, GetPinCategoryFriendlyName(PC_Enum), GeneralGrouping, Filter);
 	}
 
 	{
@@ -602,29 +590,28 @@ void UEdGraphSchema_CustomizableObject::GetGraphContextActions(FGraphContextMenu
 			ContextMenuBuilder.CreateTemplateNode<UCustomizableObjectNodeFloatConstant>(),
 			ContextMenuBuilder.CreateTemplateNode<UCustomizableObjectNodeFloatParameter>(),
 			ContextMenuBuilder.CreateTemplateNode<UCustomizableObjectNodeFloatArithmeticOp>(),
-			ContextMenuBuilder.CreateTemplateNode<UCustomizableObjectNodeFloatSwitch>(),
 			ContextMenuBuilder.CreateTemplateNode<UCustomizableObjectNodeFloatVariation>(),
 			ContextMenuBuilder.CreateTemplateNode<UCustomizableObjectNodeCurve>(),
 		};
 
 
-		AddNewNodeCategoryActionsFiltered(FloatTemplateNodes, ContextMenuBuilder, TEXT("Float"), GeneralGrouping, Filter);
+		AddNewNodeCategoryActionsFiltered(FloatTemplateNodes, ContextMenuBuilder, GetPinCategoryFriendlyName(PC_Float), GeneralGrouping, Filter);
 
 		UCustomizableObjectNodeFloatArithmeticOp* Node = NewObject<UCustomizableObjectNodeFloatArithmeticOp>();
 		Node->Operation = EFloatArithmeticOperation::E_Add;
-		AddNewNodeActionFiltered(Node, ContextMenuBuilder, TEXT("Float"), LOCTEXT("Float_Addition", "Float Addition +"), GeneralGrouping, Filter);
+		AddNewNodeActionFiltered(Node, ContextMenuBuilder, GetPinCategoryFriendlyName(PC_Float), LOCTEXT("Float_Addition", "Float Addition +"), GeneralGrouping, Filter);
 
 		Node = NewObject<UCustomizableObjectNodeFloatArithmeticOp>();
 		Node->Operation = EFloatArithmeticOperation::E_Sub;
-		AddNewNodeActionFiltered(Node, ContextMenuBuilder, TEXT("Float"), LOCTEXT("Float_Subtraction", "Float Subtraction -"), GeneralGrouping, Filter);
+		AddNewNodeActionFiltered(Node, ContextMenuBuilder, GetPinCategoryFriendlyName(PC_Float), LOCTEXT("Float_Subtraction", "Float Subtraction -"), GeneralGrouping, Filter);
 
 		Node = NewObject<UCustomizableObjectNodeFloatArithmeticOp>();
 		Node->Operation = EFloatArithmeticOperation::E_Mul;
-		AddNewNodeActionFiltered(Node, ContextMenuBuilder, TEXT("Float"), LOCTEXT("Float_Multiplication", "Float Multiplication x"), GeneralGrouping, Filter);
+		AddNewNodeActionFiltered(Node, ContextMenuBuilder, GetPinCategoryFriendlyName(PC_Float), LOCTEXT("Float_Multiplication", "Float Multiplication x"), GeneralGrouping, Filter);
 
 		Node = NewObject<UCustomizableObjectNodeFloatArithmeticOp>();
 		Node->Operation = EFloatArithmeticOperation::E_Div;
-		AddNewNodeActionFiltered(Node, ContextMenuBuilder, TEXT("Float"), LOCTEXT("Float_Division", "Float Division /"), GeneralGrouping, Filter);
+		AddNewNodeActionFiltered(Node, ContextMenuBuilder, GetPinCategoryFriendlyName(PC_Float), LOCTEXT("Float_Division", "Float Division /"), GeneralGrouping, Filter);
 	}
 
 	{
@@ -634,7 +621,7 @@ void UEdGraphSchema_CustomizableObject::GetGraphContextActions(FGraphContextMenu
 			ContextMenuBuilder.CreateTemplateNode<UCustomizableObjectNodeTransformParameter>(),
 		};
 
-		AddNewNodeCategoryActionsFiltered(TransformTemplateNodes, ContextMenuBuilder, TEXT("Transform"), GeneralGrouping, Filter);
+		AddNewNodeCategoryActionsFiltered(TransformTemplateNodes, ContextMenuBuilder, GetPinCategoryFriendlyName(PC_Transform), GeneralGrouping, Filter);
 	}
 	
 	
@@ -646,7 +633,7 @@ void UEdGraphSchema_CustomizableObject::GetGraphContextActions(FGraphContextMenu
 			ContextMenuBuilder.CreateTemplateNode<UCustomizableObjectNodeGroupProjectorParameter>(),
 		};
 
-		AddNewNodeCategoryActionsFiltered(ProjectorTemplateNodes, ContextMenuBuilder, TEXT("Projector"), GeneralGrouping, Filter);
+		AddNewNodeCategoryActionsFiltered(ProjectorTemplateNodes, ContextMenuBuilder, GetPinCategoryFriendlyName(PC_Projector), GeneralGrouping, Filter);
 	}
 	
 	{
@@ -661,7 +648,7 @@ void UEdGraphSchema_CustomizableObject::GetGraphContextActions(FGraphContextMenu
 			ContextMenuBuilder.CreateTemplateNode<UCONodeSkeletalMeshBreak>(),
 		};
 
-		AddNewNodeCategoryActionsFiltered(ExperimentalTemplateNodes, ContextMenuBuilder, NC_Experimental.ToString(), 2, Filter);
+		AddNewNodeCategoryActionsFiltered(ExperimentalTemplateNodes, ContextMenuBuilder, NC_Experimental, 2, Filter);
 	}
 
 	{
@@ -669,17 +656,18 @@ void UEdGraphSchema_CustomizableObject::GetGraphContextActions(FGraphContextMenu
 		{
 			ContextMenuBuilder.CreateTemplateNode<UCONodeMaterialConstant>(),
 			ContextMenuBuilder.CreateTemplateNode<UCONodeMaterialVariation>(),
-			ContextMenuBuilder.CreateTemplateNode<UCONodeMaterialSwitch>(),
 			ContextMenuBuilder.CreateTemplateNode<UCONodeMaterialBreak>(),
 			ContextMenuBuilder.CreateTemplateNode<UCONodeMaterialModify>()
 		};
 
-		AddNewNodeCategoryActionsFiltered(MaterialTemplateNodes, ContextMenuBuilder, NC_Material.ToString(), 2, Filter);
+		AddNewNodeCategoryActionsFiltered(MaterialTemplateNodes, ContextMenuBuilder, NC_Material, 2, Filter);
 	}
 	
 	// Extensions
-	TArray<const UScriptStruct*> ExternalTypes;
+	ICustomizableObjectEditorModulePrivate& EditorModule = ICustomizableObjectEditorModulePrivate::GetChecked();
 
+	const TMap<FName, TStrongObjectPtr<const UScriptStruct>>& ExternalPinTypes = EditorModule.GetExternalPinTypes();
+	
 	{
 		for (const TStrongObjectPtr<const UScriptStruct>& OperationType : ICustomizableObjectModulePrivate::Get().GetRegisteredExternalOperations())
 		{
@@ -693,70 +681,67 @@ void UEdGraphSchema_CustomizableObject::GetGraphContextActions(FGraphContextMenu
 					UCONodeExternalOperation* Node = ContextMenuBuilder.CreateTemplateNode<UCONodeExternalOperation>();
 					Node->OperationInstancedStruct = OperationInstancedStruct;
 
-					AddNewNodeActionFiltered(Node, ContextMenuBuilder, GetPinCategoryName(Operation->GetOutput().Value).ToString(), GeneralGrouping, Filter);
-
-					ExternalTypes.AddUnique(Operation->GetOutput().Value);
-
-					for (const TPair<FName, const UScriptStruct*>& Pair : Operation->GetInputs())
-					{
-						ExternalTypes.AddUnique(Pair.Value);
-					}
+					AddNewNodeActionFiltered(Node, ContextMenuBuilder, GetPinCategoryFriendlyName(GetPinCategoryName(Operation->GetOutput().Value)), GeneralGrouping, Filter);
 				}
 			}
 		}
-
+		
 		// Parameter
-		for (const UScriptStruct* ScriptStruct : ExternalTypes)
+		for (const TTuple<FName, TStrongObjectPtr<const UScriptStruct>>& Pair : ExternalPinTypes)
 		{
-			if (IsTypeNative(ScriptStruct))
+			if (IsTypeNative(Pair.Value.Get()))
 			{
 				continue;
 			}
-
-			if (ScriptStruct->GetCppStructOps()->GetCapabilities().IsAbstract)
-			{
-				continue;
-			}
-				
+			
 			UCONodeExternalTypeParameter* Node = ContextMenuBuilder.CreateTemplateNode<UCONodeExternalTypeParameter>();
-			Node->DefaultValue.InitializeAs(ScriptStruct);
+			Node->DefaultValue.InitializeAs(Pair.Value.Get());
 
-			AddNewNodeActionFiltered(Node, ContextMenuBuilder, GetPinCategoryName(ScriptStruct).ToString(), GeneralGrouping, Filter);
+			AddNewNodeActionFiltered(Node, ContextMenuBuilder, GetPinCategoryFriendlyName(Pair.Key), GeneralGrouping, Filter);
 		}
 	}
 	
 	// External Pin Nodes
 	{
-		TArray<FName> PinTypes({ PC_MeshSection, PC_Modifier, PC_Mesh, PC_Texture, PC_Projector, PC_GroupProjector, PC_Color, PC_Float, PC_Bool, PC_Enum, PC_Transform, PC_Stack, PC_Texture_Passthrough, PC_Texture_Passthrough, PC_Material, PC_PoseAsset, PC_Component });
+		TArray<FName> PinTypes({ PC_MeshSection, PC_Mesh, PC_Texture, PC_Projector, PC_Color, PC_Float, PC_Transform, PC_Texture_Passthrough, PC_Material, PC_Component });
+		
+		// External Types
+		for (const TTuple<FName, TStrongObjectPtr<const UScriptStruct>>& Pair : ExternalPinTypes)
+		{
+			PinTypes.AddUnique(Pair.Key);
+		}
 
+		for (const FName& PinCategory : PinTypes)
+		{
+			UCONodeSwitch* NodeSwitch = ContextMenuBuilder.CreateTemplateNode<UCONodeSwitch>();
+			NodeSwitch->PinType = PinCategory;
+			
+			const FText CategoryFriendlyName = GetPinCategoryFriendlyName(PinCategory);
+			
+			AddNewNodeActionFiltered(NodeSwitch, ContextMenuBuilder, CategoryFriendlyName, GeneralGrouping, Filter);
+		}
+		
+		PinTypes.Append({  PC_Stack, PC_Modifier, PC_Projector, PC_GroupProjector, PC_PoseAsset, PC_Enum });
+		
 		// Extensions
 		for (const FRegisteredCustomizableObjectPinType& PinType : ICustomizableObjectModule::Get().GetExtendedPinTypes())
 		{
 			PinTypes.AddUnique(PinType.PinType.Name);
 		}
-
-		// External Types
-		for (const UScriptStruct* ScriptStruct : ExternalTypes)
-		{
-			const FName Category = GetPinCategoryName(ScriptStruct);
-			PinTypes.AddUnique(Category);
-			FInstancedStruct OperationInstancedStruct;
-		}
-
+		
 		for (const FName& PinCategory : PinTypes)
 		{
+			const FText CategoryFriendlyName = GetPinCategoryFriendlyName(PinCategory);
+
 			UCustomizableObjectNodeExternalPin* CustomizableObjectNodeExternalPin = ContextMenuBuilder.CreateTemplateNode<UCustomizableObjectNodeExternalPin>();
 			CustomizableObjectNodeExternalPin->PinType = PinCategory;
 			
-			AddNewNodeActionFiltered(CustomizableObjectNodeExternalPin, ContextMenuBuilder, TEXT("Import Pin"), GeneralGrouping, Filter);
-		}
-
-		for (const FName& PinCategory : PinTypes)
-		{
+			AddNewNodeActionFiltered(CustomizableObjectNodeExternalPin, ContextMenuBuilder, CategoryFriendlyName, GeneralGrouping, Filter);
+			
 			UCustomizableObjectNodeExposePin* CustomizableObjectNodeExposePin = ContextMenuBuilder.CreateTemplateNode<UCustomizableObjectNodeExposePin>();
 			CustomizableObjectNodeExposePin->PinType = PinCategory;
 			
-			AddNewNodeActionFiltered(CustomizableObjectNodeExposePin, ContextMenuBuilder, TEXT("Export Pin"), GeneralGrouping, Filter);
+			AddNewNodeActionFiltered(CustomizableObjectNodeExposePin, ContextMenuBuilder, CategoryFriendlyName, GeneralGrouping, Filter);
 		}
 	}
 
@@ -778,20 +763,20 @@ void UEdGraphSchema_CustomizableObject::GetGraphContextActions(FGraphContextMenu
 		if (Node->ShouldAddToContextMenu(Category))
 		{
 			UCustomizableObjectNode* TemplateNode = ContextMenuBuilder.CreateTemplateNode<UCustomizableObjectNode>(Node->GetClass());
-			AddNewNodeActionFiltered(TemplateNode, ContextMenuBuilder, Category.ToString(), GeneralGrouping, Filter);
+			AddNewNodeActionFiltered(TemplateNode, ContextMenuBuilder, Category, GeneralGrouping, Filter);
 		}
 	}
 
 	if (!ContextMenuBuilder.FromPin)
 	{
 		UEdGraphNode_Comment* Node = NewObject<UEdGraphNode_Comment>();
-		TSharedPtr<FCustomizableObjectSchemaAction_NewNode> Action = AddNewNodeAction(ContextMenuBuilder, FString(), Node->GetNodeTitle(ENodeTitleType::ListView), FText(), 1);
+		TSharedPtr<FCustomizableObjectSchemaAction_NewNode> Action = AddNewNodeAction(ContextMenuBuilder, {}, Node->GetNodeTitle(ENodeTitleType::ListView), FText(), 1);
 		Action->NodeTemplate = Node;
 	}
 	
 	{
 		UCustomizableObjectNodeReroute* Node = NewObject<UCustomizableObjectNodeReroute>();
-		TSharedPtr<FCustomizableObjectSchemaAction_NewNode> Action = AddNewNodeAction(ContextMenuBuilder, FString(), Node->GetNodeTitle(ENodeTitleType::ListView), FText(), 1);
+		TSharedPtr<FCustomizableObjectSchemaAction_NewNode> Action = AddNewNodeAction(ContextMenuBuilder, {}, Node->GetNodeTitle(ENodeTitleType::ListView), FText(), 1);
 		Action->NodeTemplate = Node;
 	}
 
@@ -877,7 +862,7 @@ const FPinConnectionResponse UEdGraphSchema_CustomizableObject::CanCreateConnect
 }
 
 
-FLinearColor UEdGraphSchema_CustomizableObject::GetPinTypeColor(const FName& TypeString)
+FLinearColor UEdGraphSchema_CustomizableObject::GetPinTypeColor(const FName& Type)
 {
 	static TMap<FName, FLinearColor> PinCategoryColors;
 
@@ -927,21 +912,27 @@ FLinearColor UEdGraphSchema_CustomizableObject::GetPinTypeColor(const FName& Typ
 		}
 	}
 		
-	if (const FLinearColor* Color = PinCategoryColors.Find(TypeString))
+	if (const FLinearColor* Color = PinCategoryColors.Find(Type))
 	{
 		return *Color;		
 	}
 
 	for (const FRegisteredCustomizableObjectPinType& PinType : ICustomizableObjectModule::Get().GetExtendedPinTypes())
 	{
-		if (PinType.PinType.Name == TypeString)
+		if (PinType.PinType.Name == Type)
 		{
 			return PinType.PinType.Color;
 		}
 	}
+	
+	const TMap<FName, FLinearColor>& ExternalPinColors = ICustomizableObjectEditorModulePrivate::Get()->GetExternalPinTypeColors();
+	if (const FLinearColor* Result = ExternalPinColors.Find(Type))
+	{
+		return *Result;
+	}
 
-	const uint32 Hash = GetTypeHash(TypeString.ToString());
-	return FLinearColor::MakeRandomSeededColor(Hash);
+	unimplemented();
+	return FLinearColor::Black;
 }
 
 
@@ -1801,31 +1792,22 @@ FText UEdGraphSchema_CustomizableObject::GetPinCategoryFriendlyName(const FName&
 	{
 		return LOCTEXT("String_Pin_Category", "String");
 	}
-	else
+	
+	for (const FRegisteredCustomizableObjectPinType& PinType : ICustomizableObjectModule::Get().GetExtendedPinTypes())
 	{
-		for (const FRegisteredCustomizableObjectPinType& PinType : ICustomizableObjectModule::Get().GetExtendedPinTypes())
+		if (PinType.PinType.Name == PinCategory)
 		{
-			if (PinType.PinType.Name == PinCategory)
-			{
-				return PinType.PinType.DisplayName;
-			}
+			return PinType.PinType.DisplayName;
 		}
-
-		return {};
-	}
-}
-
-
-FText UEdGraphSchema_CustomizableObject::GetPinCategoryFriendlyName(const UScriptStruct* Type)
-{
-	if (IsTypeNative(Type))
-	{
-		return GetPinCategoryFriendlyName(GetPinCategoryName(Type));
 	}
-	else
+	
+	const TMap<FName, FText>& ExternalPinFriendlyName = ICustomizableObjectEditorModulePrivate::Get()->GetExternalPinFriendlyName();
+	if (const FText* Result = ExternalPinFriendlyName.Find(PinCategory))
 	{
-		return Type->GetDisplayNameText();
+		return *Result;
 	}
+	
+	return {};
 }
 
 
diff --git a/Engine/Plugins/Mutable/Source/CustomizableObjectEditor/Private/MuCOE/GenerateMutableSource/GenerateMutableSource.cpp b/Engine/Plugins/Mutable/Source/CustomizableObjectEditor/Private/MuCOE/GenerateMutableSource/GenerateMutableSource.cpp
--- a/Engine/Plugins/Mutable/Source/CustomizableObjectEditor/Private/MuCOE/GenerateMutableSource/GenerateMutableSource.cpp
+++ b/Engine/Plugins/Mutable/Source/CustomizableObjectEditor/Private/MuCOE/GenerateMutableSource/GenerateMutableSource.cpp
@@ -33,8 +33,6 @@
 #include "MuCOE/Nodes/CustomizableObjectNodeTextureLayer.h"
 #include "MuCOE/Nodes/CustomizableObjectNodeTextureParameter.h"
 #include "MuCOE/Nodes/CustomizableObjectNodeTextureProject.h"
-#include "MuCOE/Nodes/CustomizableObjectNodeTextureSwitch.h"
-#include "MuCOE/Nodes/CustomizableObjectNodePassThroughTextureSwitch.h"
 #include "MuCOE/Nodes/CustomizableObjectNodeTextureToChannels.h"
 #include "MuCOE/Nodes/CustomizableObjectNodeTextureTransform.h"
 #include "MuCOE/Nodes/CustomizableObjectNodeTextureSaturate.h"
@@ -51,8 +49,8 @@
 #include "PlatformInfo.h"
 #include "Math/NumericLimits.h"
 #include "Hash/CityHash.h"
-#include "MuCO/CustomizableObjectResourceDataTypes.h"
 #include "MuCOE/Nodes/CustomizableObjectNodeGroupProjectorParameter.h"
+#include "MuCOE/Nodes/CONodeSwitch.h"
 
 #define LOCTEXT_NAMESPACE "CustomizableObjectEditor"
 
@@ -496,7 +494,7 @@ UTexture2D* FindReferenceImage(const UEdGraphPin* Pin, FMutableGraphGenerationCo
 		}
 	}
 
-	else if (const UCustomizableObjectNodeTextureSwitch* TypedNodeSwitch = Cast<UCustomizableObjectNodeTextureSwitch>(Node))
+	else if (const UCONodeSwitch* TypedNodeSwitch = CastSwitch(Node, UEdGraphSchema_CustomizableObject::PC_Texture))
 	{
 		for (int SelectorIndex = 0; !Result && SelectorIndex < TypedNodeSwitch->SwitchPins.Num(); ++SelectorIndex)
 		{
@@ -507,7 +505,7 @@ UTexture2D* FindReferenceImage(const UEdGraphPin* Pin, FMutableGraphGenerationCo
 		}
 	}
 
-	else if (const UCustomizableObjectNodePassThroughTextureSwitch* TypedNodePassThroughSwitch = Cast<UCustomizableObjectNodePassThroughTextureSwitch>(Node))
+	else if (const UCONodeSwitch* TypedNodePassThroughSwitch = CastSwitch(Node, UEdGraphSchema_CustomizableObject::PC_Texture_Passthrough))
 	{
 		for (int32 SelectorIndex = 0; !Result && SelectorIndex < TypedNodePassThroughSwitch->SwitchPins.Num(); ++SelectorIndex)
 		{
@@ -835,7 +833,7 @@ uint32 AddUniqueSurfaceMetadata(const FMutableSourceSurfaceMetadata& SourceMetad
 }
 
 
-bool DoOptionsMatchEnum(const UCustomizableObjectNodeSwitchBase& InEditorSwitchNode,
+bool DoOptionsMatchEnum(const UCONodeSwitch& InEditorSwitchNode,
 	const UE::Mutable::Private::NodeScalarEnumParameter& MutableScalarParameter)
 {
 	// todo: use the actual UEditor CO node once the issue UE-294591 is resolved. Currently we do need to use the core node as GenerateMutableSource is
diff --git a/Engine/Plugins/Mutable/Source/CustomizableObjectEditor/Private/MuCOE/GenerateMutableSource/GenerateMutableSource.h b/Engine/Plugins/Mutable/Source/CustomizableObjectEditor/Private/MuCOE/GenerateMutableSource/GenerateMutableSource.h
--- a/Engine/Plugins/Mutable/Source/CustomizableObjectEditor/Private/MuCOE/GenerateMutableSource/GenerateMutableSource.h
+++ b/Engine/Plugins/Mutable/Source/CustomizableObjectEditor/Private/MuCOE/GenerateMutableSource/GenerateMutableSource.h
@@ -44,7 +44,7 @@
 
 class UCustomizableObjectNodeGroupProjectorParameter;
 class UCONodeSkeletalMeshSectionBase;
-class UCustomizableObjectNodeSwitchBase;
+class UCONodeSwitch;
 class UCustomizableObjectNodeModifierClipWithMesh;
 
 class FCustomizableObjectCompiler;
@@ -1025,5 +1025,5 @@ UE::Tasks::FTask ConvertSkeletalMeshToMutable(TSharedRef<FMeshConversionContext>
 * @param MutableScalarParameter The enum parameter node we want to compare the options against.
  * @return true if the amount of options does match and false otherwise.
  */
-bool DoOptionsMatchEnum(const UCustomizableObjectNodeSwitchBase& InEditorSwitchNode, const UE::Mutable::Private::NodeScalarEnumParameter& MutableScalarParameter);
+bool DoOptionsMatchEnum(const UCONodeSwitch& InEditorSwitchNode, const UE::Mutable::Private::NodeScalarEnumParameter& MutableScalarParameter);
 
diff --git a/Engine/Plugins/Mutable/Source/CustomizableObjectEditor/Private/MuCOE/GenerateMutableSource/GenerateMutableSourceColor.cpp b/Engine/Plugins/Mutable/Source/CustomizableObjectEditor/Private/MuCOE/GenerateMutableSource/GenerateMutableSourceColor.cpp
--- a/Engine/Plugins/Mutable/Source/CustomizableObjectEditor/Private/MuCOE/GenerateMutableSource/GenerateMutableSourceColor.cpp
+++ b/Engine/Plugins/Mutable/Source/CustomizableObjectEditor/Private/MuCOE/GenerateMutableSource/GenerateMutableSourceColor.cpp
@@ -18,12 +18,12 @@
 #include "MuCOE/Nodes/CustomizableObjectNodeColorConstant.h"
 #include "MuCOE/Nodes/CustomizableObjectNodeColorFromFloats.h"
 #include "MuCOE/Nodes/CustomizableObjectNodeColorParameter.h"
-#include "MuCOE/Nodes/CustomizableObjectNodeColorSwitch.h"
 #include "MuCOE/Nodes/CustomizableObjectNodeColorVariation.h"
 #include "MuCOE/Nodes/CustomizableObjectNodeColorToSRGB.h"
 #include "MuCOE/Nodes/CustomizableObjectNodeMacroInstance.h"
 #include "MuCOE/Nodes/CONodeMaterialBreak.h"
 #include "MuCOE/Nodes/CustomizableObjectNodeEnumParameter.h"
+#include "MuCOE/Nodes/CONodeSwitch.h"
 #include "MuCOE/Nodes/CustomizableObjectNodeTable.h"
 #include "MuCOE/Nodes/CustomizableObjectNodeTextureSample.h"
 #include "MuCOE/Nodes/CustomizableObjectNodeTunnel.h"
@@ -93,7 +93,7 @@ UE::Mutable::Private::Ptr<UE::Mutable::Private::NodeColour> GenerateMutableSourc
 			EMutableParameterType::Color));
 	}
 
-	else if (const UCustomizableObjectNodeColorSwitch* TypedNodeColorSwitch = Cast<UCustomizableObjectNodeColorSwitch>(Node))
+	else if (const UCONodeSwitch* TypedNodeColorSwitch = CastSwitch(Node, UEdGraphSchema_CustomizableObject::PC_Color))
 	{
 		Result = [&]() -> UE::Mutable::Private::Ptr<UE::Mutable::Private::NodeColourSwitch>
 		{
diff --git a/Engine/Plugins/Mutable/Source/CustomizableObjectEditor/Private/MuCOE/GenerateMutableSource/GenerateMutableSourceComponent.cpp b/Engine/Plugins/Mutable/Source/CustomizableObjectEditor/Private/MuCOE/GenerateMutableSource/GenerateMutableSourceComponent.cpp
--- a/Engine/Plugins/Mutable/Source/CustomizableObjectEditor/Private/MuCOE/GenerateMutableSource/GenerateMutableSourceComponent.cpp
+++ b/Engine/Plugins/Mutable/Source/CustomizableObjectEditor/Private/MuCOE/GenerateMutableSource/GenerateMutableSourceComponent.cpp
@@ -15,7 +15,6 @@
 #include "MuCOE/Nodes/CustomizableObjectNodeComponentMesh.h"
 #include "MuCOE/Nodes/CustomizableObjectNodeComponentMeshAddTo.h"
 #include "MuCOE/Nodes/CustomizableObjectNodeComponentPassthroughMesh.h"
-#include "MuCOE/Nodes/CustomizableObjectNodeComponentSwitch.h"
 #include "MuCOE/Nodes/CustomizableObjectNodeComponentVariation.h"
 #include "MuCOE/Nodes/CustomizableObjectNodeModifierBase.h"
 #include "MuCOE/Nodes/CustomizableObjectNodeMacroInstance.h"
@@ -32,6 +31,7 @@
 #include "Rendering/SkeletalMeshLODModel.h"
 #include "Interfaces/ITargetPlatform.h"
 #include "MuCOE/EdGraphSchema_CustomizableObject.h"
+#include "MuCOE/Nodes/CONodeSwitch.h"
 
 
 #define LOCTEXT_NAMESPACE "CustomizableObjectEditor"
@@ -227,7 +227,7 @@ UE::Mutable::Private::Ptr<UE::Mutable::Private::NodeComponent> GenerateMutableSo
 		Result = ComponentNode;
 	}
 	
-	else if (const UCustomizableObjectNodeComponentSwitch* TypedNodeSwitch = Cast<UCustomizableObjectNodeComponentSwitch>(Node))
+	else if (const UCONodeSwitch* TypedNodeSwitch = CastSwitch(Node, UEdGraphSchema_CustomizableObject::PC_Component))
 	{
 		// Using a lambda so control flow is easier to manage.
 		Result = [&]()
diff --git a/Engine/Plugins/Mutable/Source/CustomizableObjectEditor/Private/MuCOE/GenerateMutableSource/GenerateMutableSourceExternal.cpp b/Engine/Plugins/Mutable/Source/CustomizableObjectEditor/Private/MuCOE/GenerateMutableSource/GenerateMutableSourceExternal.cpp
--- a/Engine/Plugins/Mutable/Source/CustomizableObjectEditor/Private/MuCOE/GenerateMutableSource/GenerateMutableSourceExternal.cpp
+++ b/Engine/Plugins/Mutable/Source/CustomizableObjectEditor/Private/MuCOE/GenerateMutableSource/GenerateMutableSourceExternal.cpp
@@ -13,6 +13,7 @@
 #include "MuCOE/GraphTraversal.h"
 #include "MuCOE/Nodes/CONodeExternalOperation.h"
 #include "MuCOE/Nodes/CONodeExternalTypeParameter.h"
+#include "MuCOE/Nodes/CONodeSwitch.h"
 #include "MuCOE/Nodes/ExternalTypeCustomization.h"
 #include "MuR/External/FloatAdapter.h"
 #include "MuR/External/MaterialAdapter.h"
@@ -22,6 +23,7 @@
 #include "MuR/External/VectorAdapter.h"
 #include "MuT/NodeExternalOperation.h"
 #include "MuT/NodeExternalParameter.h"
+#include "MuT/NodeExternalSwitch.h"
 
 #define LOCTEXT_NAMESPACE "CustomizableObjectEditor"
 
@@ -61,15 +63,7 @@ UE::Mutable::Private::Ptr<UE::Mutable::Private::NodeExternal> GenerateMutableSou
 			const UE::Mutable::FExternalOperation* Operation = TypedNodeExternal->OperationInstancedStruct.GetPtr<UE::Mutable::FExternalOperation>();
 			if (!Operation)
 			{
-				FText Msg = FText::Format(LOCTEXT("ExternalNotLoaded", "Node is using an external operation which has not been loaded: {0}"), FText::FromName(TypedNodeExternal->CachedOperationName));
-				GenerationContext.Log(Msg, TypedNodeExternal, EMessageSeverity::Error);
-				return;
-			}
-
-			const UScriptStruct* ExternalType;
-			if (ContainsUObject(*Operation, ExternalType))
-			{
-				FText Msg = FText::Format(LOCTEXT("TypeUObject", "Operation contains UObject properties. Operation: {0} Type: {1}"), FText::FromName(TypedNodeExternal->CachedOperationName), FText::FromName(ExternalType->GetFName()));
+				FText Msg = FText::Format(LOCTEXT("ExternalNotLoaded", "Node is using an external operation which has not been loaded: {0}"), TypedNodeExternal->CachedOperationName);
 				GenerationContext.Log(Msg, TypedNodeExternal, EMessageSeverity::Error);
 				return;
 			}
@@ -78,19 +72,19 @@ UE::Mutable::Private::Ptr<UE::Mutable::Private::NodeExternal> GenerateMutableSou
             				
             ExternalNode->OperationInstancedStruct = TypedNodeExternal->OperationInstancedStruct;
 
-            TArray<TPair<FName, const UScriptStruct*>> Inputs = Operation->GetInputs();
-            for (const TPair<FName, const UScriptStruct*>& Pair : Inputs)
+            TArray<TPair<FText, const UScriptStruct*>> Inputs = Operation->GetInputs();
+            for (const TPair<FText, const UScriptStruct*>& Pair : Inputs)
             {
-            	const FName& InputName = Pair.Key;
+            	const FText& InputName = Pair.Key;
             
             	const FEdGraphPinReference* InputResult = TypedNodeExternal->InputPins.FindByPredicate([&](const FEdGraphPinReference& Element)
             	{
-            		return TypedNodeExternal->GetPinData<UCONodeExtenalPinData>(*Element.Get()).Name == InputName;
+            		return Element.Get()->PinFriendlyName.EqualTo(InputName);
             	});
         
             	if (!InputResult)
             	{
-            		FText Msg = FText::Format(LOCTEXT("ExternalInputNotFound", "Can not find input in external operation: {0}"), FText::FromName(InputName));
+            		FText Msg = FText::Format(LOCTEXT("ExternalInputNotFound", "Can not find input in external operation: {0}"), InputName);
             		GenerationContext.Log(Msg, TypedNodeExternal, EMessageSeverity::Error);
             		continue;
             	}
@@ -98,7 +92,7 @@ UE::Mutable::Private::Ptr<UE::Mutable::Private::NodeExternal> GenerateMutableSou
             	const UEdGraphPin* ConnectedPin = FollowInputPin(*InputResult->Get());
             	if (!ConnectedPin)
             	{
-            		FText Msg = FText::Format(LOCTEXT("ExternalInputNotConnected", "External operation input not connected: {0}"), FText::FromName(InputName));
+            		FText Msg = FText::Format(LOCTEXT("ExternalInputNotConnected", "External operation input not connected: {0}"), InputName);
             		GenerationContext.Log(Msg, TypedNodeExternal, EMessageSeverity::Error);
             		continue;
             	}
@@ -144,6 +138,78 @@ UE::Mutable::Private::Ptr<UE::Mutable::Private::NodeExternal> GenerateMutableSou
 		}();		
 	}
 
+	if (const UCONodeSwitch* NodeSwitch = Cast<UCONodeSwitch>(EditorNode))
+	{
+		// Using a lambda so control flow is easier to manage.
+		Result = [&]() -> UE::Mutable::Private::Ptr<UE::Mutable::Private::NodeExternal>
+		{
+			const UEdGraphPin* SwitchParameter = NodeSwitch->SwitchParameterPinReference.Get();
+
+			// Check Switch Parameter arity preconditions.
+			if (const int32 NumParameters = FollowInputPinArray(*SwitchParameter).Num();
+				NumParameters != 1)
+			{
+				const FText Message = NumParameters
+					? LOCTEXT("NoEnumParamInSwitch", "Switch nodes must have an enum switch parameter. Please connect an enum and refresh the switch node.")
+					: LOCTEXT("InvalidEnumInSwitch", "Switch nodes must have a single enum with all the options inside. Please remove all the enums but one and refresh the switch node.");
+
+				GenerationContext.Log(Message, Node);
+				return nullptr;
+			}
+
+			const UEdGraphPin* EnumPin = FollowInputPin(*SwitchParameter);
+			UE::Mutable::Private::Ptr<UE::Mutable::Private::NodeScalar> SwitchParam = GenerateMutableSourceFloat(EnumPin, GenerationContext);
+
+			// Switch Param not generated
+			if (!SwitchParam)
+			{
+				// Warn about a failure.
+				if (EnumPin)
+				{
+					const FText Message = LOCTEXT("FailedToGenerateSwitchParam", "Could not generate switch enum parameter. Please refresh the switch node and connect an enum.");
+					GenerationContext.Log(Message, Node);
+				}
+
+				return nullptr;
+			}
+
+			if (SwitchParam->GetType() != UE::Mutable::Private::NodeScalarEnumParameter::GetStaticType())
+			{
+				const FText Message = LOCTEXT("WrongSwitchParamType", "Switch parameter of incorrect type.");
+				GenerationContext.Log(Message, Node);
+
+				return nullptr;
+			}
+
+			{
+				const UE::Mutable::Private::NodeScalarEnumParameter* EnumParameter = static_cast<UE::Mutable::Private::NodeScalarEnumParameter*>(SwitchParam.get());
+				
+				if (!DoOptionsMatchEnum(*NodeSwitch, *EnumParameter))
+				{
+					const FText Message = LOCTEXT("MismatchedSwitch", "Switch enum and switch node have different options. Please refresh the switch node to make sure the outcomes are labeled properly.");
+					GenerationContext.Log(Message, Node);
+					Node->SetRefreshNodeWarning();
+				}
+			}
+
+			const int32 NumSwitchOptions = NodeSwitch->SwitchPins.Num();
+				
+			UE::Mutable::Private::Ptr<UE::Mutable::Private::NodeExternalSwitch> SwitchNode = new UE::Mutable::Private::NodeExternalSwitch;
+			SwitchNode->Parameter = SwitchParam;
+			SwitchNode->Options.SetNum(NumSwitchOptions);
+
+			for (int SelectorIndex = 0; SelectorIndex < NumSwitchOptions; ++SelectorIndex)
+			{
+				if (const UEdGraphPin* ConnectedPin = FollowInputPin(*NodeSwitch->SwitchPins[SelectorIndex].Get()))
+				{
+					SwitchNode->Options[SelectorIndex] = GenerateMutableSourceExternal(ConnectedPin, GenerationContext, Options);
+				}
+			}
+
+			return SwitchNode;
+		}(); // invoke lambda.
+	}
+	
 	if (const UCONodeExternalTypeParameter* TypeNodeExternalParameter = Cast<UCONodeExternalTypeParameter>(EditorNode))
 	{
 		FString ParameterName = TypeNodeExternalParameter->GetParameterName(&GenerationContext.MacroNodesStack);
diff --git a/Engine/Plugins/Mutable/Source/CustomizableObjectEditor/Private/MuCOE/GenerateMutableSource/GenerateMutableSourceFloat.cpp b/Engine/Plugins/Mutable/Source/CustomizableObjectEditor/Private/MuCOE/GenerateMutableSource/GenerateMutableSourceFloat.cpp
--- a/Engine/Plugins/Mutable/Source/CustomizableObjectEditor/Private/MuCOE/GenerateMutableSource/GenerateMutableSourceFloat.cpp
+++ b/Engine/Plugins/Mutable/Source/CustomizableObjectEditor/Private/MuCOE/GenerateMutableSource/GenerateMutableSourceFloat.cpp
@@ -20,10 +20,10 @@
 #include "MuCOE/Nodes/CustomizableObjectNodeEnumParameter.h"
 #include "MuCOE/Nodes/CustomizableObjectNodeFloatConstant.h"
 #include "MuCOE/Nodes/CustomizableObjectNodeFloatParameter.h"
-#include "MuCOE/Nodes/CustomizableObjectNodeFloatSwitch.h"
 #include "MuCOE/Nodes/CustomizableObjectNodeFloatVariation.h"
 #include "MuCOE/Nodes/CustomizableObjectNodeFloatArithmeticOp.h"
 #include "MuCOE/Nodes/CustomizableObjectNodeMacroInstance.h"
+#include "MuCOE/Nodes/CONodeSwitch.h"
 #include "MuCOE/Nodes/CustomizableObjectNodeTable.h"
 #include "MuCOE/Nodes/CustomizableObjectNodeTunnel.h"
 #include "MuT/NodeScalarConstant.h"
@@ -119,7 +119,7 @@ UE::Mutable::Private::Ptr<UE::Mutable::Private::NodeScalar> GenerateMutableSourc
 		GenerationContext.ParameterUIDataMap.Add(EnumParamNode->GetParameterName(&GenerationContext.MacroNodesStack), ParameterUIData);
 	}
 
-	else if (const UCustomizableObjectNodeFloatSwitch* TypedNodeFloatSwitch = Cast<UCustomizableObjectNodeFloatSwitch>(Node))
+	else if (const UCONodeSwitch* TypedNodeFloatSwitch = CastSwitch(Node, UEdGraphSchema_CustomizableObject::PC_Float))
 	{
 		// Using a lambda so control flow is easier to manage.
 		Result = [&]() -> UE::Mutable::Private::Ptr<UE::Mutable::Private::NodeScalar>
diff --git a/Engine/Plugins/Mutable/Source/CustomizableObjectEditor/Private/MuCOE/GenerateMutableSource/GenerateMutableSourceImage.cpp b/Engine/Plugins/Mutable/Source/CustomizableObjectEditor/Private/MuCOE/GenerateMutableSource/GenerateMutableSourceImage.cpp
--- a/Engine/Plugins/Mutable/Source/CustomizableObjectEditor/Private/MuCOE/GenerateMutableSource/GenerateMutableSourceImage.cpp
+++ b/Engine/Plugins/Mutable/Source/CustomizableObjectEditor/Private/MuCOE/GenerateMutableSource/GenerateMutableSourceImage.cpp
@@ -33,8 +33,6 @@
 #include "MuCOE/Nodes/CustomizableObjectNodeTextureLayer.h"
 #include "MuCOE/Nodes/CustomizableObjectNodeTextureParameter.h"
 #include "MuCOE/Nodes/CustomizableObjectNodeTextureProject.h"
-#include "MuCOE/Nodes/CustomizableObjectNodeTextureSwitch.h"
-#include "MuCOE/Nodes/CustomizableObjectNodePassThroughTextureSwitch.h"
 #include "MuCOE/Nodes/CustomizableObjectNodeTextureToChannels.h"
 #include "MuCOE/Nodes/CustomizableObjectNodeTextureTransform.h"
 #include "MuCOE/Nodes/CustomizableObjectNodeTextureSaturate.h"
@@ -69,6 +67,7 @@
 #include "Interfaces/ITextureFormat.h"
 #include "Interfaces/ITextureFormatManagerModule.h"
 #include "MuCOE/Nodes/CustomizableObjectNodeEnumParameter.h"
+#include "MuCOE/Nodes/CONodeSwitch.h"
 #include "MuT/UnrealPixelFormatOverride.h"
 
 #define LOCTEXT_NAMESPACE "CustomizableObjectEditor"
@@ -346,7 +345,7 @@ UE::Mutable::Private::Ptr<UE::Mutable::Private::NodeImage> GenerateMutableSource
 		}
 	}
 
-	else if (const UCustomizableObjectNodeTextureSwitch* TypedNodeTextureSwitch = Cast<UCustomizableObjectNodeTextureSwitch>(Node))
+	else if (const UCONodeSwitch* TypedNodeTextureSwitch = CastSwitch(Node, UEdGraphSchema_CustomizableObject::PC_Texture))
 	{
 		Result = [&]()
 		{
@@ -918,7 +917,7 @@ UE::Mutable::Private::Ptr<UE::Mutable::Private::NodeImage> GenerateMutableSource
 		}
 	}
 
-	else if (const UCustomizableObjectNodePassThroughTextureSwitch* TypedNodePassThroughTextureSwitch = Cast<UCustomizableObjectNodePassThroughTextureSwitch>(Node))
+	else if (const UCONodeSwitch* TypedNodePassThroughTextureSwitch = CastSwitch(Node, UEdGraphSchema_CustomizableObject::PC_Texture_Passthrough))
 	{
 		Result = [&]()
 		{
diff --git a/Engine/Plugins/Mutable/Source/CustomizableObjectEditor/Private/MuCOE/GenerateMutableSource/GenerateMutableSourceMaterial.cpp b/Engine/Plugins/Mutable/Source/CustomizableObjectEditor/Private/MuCOE/GenerateMutableSource/GenerateMutableSourceMaterial.cpp
--- a/Engine/Plugins/Mutable/Source/CustomizableObjectEditor/Private/MuCOE/GenerateMutableSource/GenerateMutableSourceMaterial.cpp
+++ b/Engine/Plugins/Mutable/Source/CustomizableObjectEditor/Private/MuCOE/GenerateMutableSource/GenerateMutableSourceMaterial.cpp
@@ -26,7 +26,6 @@
 #include "MuCOE/Nodes/CONodeMaterialBreak.h"
 #include "MuCOE/Nodes/CONodeMaterialConstant.h"
 #include "MuCOE/Nodes/CONodeMaterialModify.h"
-#include "MuCOE/Nodes/CONodeMaterialSwitch.h"
 #include "MuCOE/Nodes/CONodeMaterialVariation.h"
 #include "MuCO/LoadUtils.h"
 #include "MuT/NodeImageFormat.h"
@@ -44,6 +43,7 @@
 #include "Materials/MaterialInstance.h"
 #include "MuCOE/Nodes/CONodeExternalOperation.h"
 #include "MuCOE/Nodes/CONodeSkeletalMeshBreak.h"
+#include "MuCOE/Nodes/CONodeSwitch.h"
 #include "MuT/NodeMaterialExternal.h"
 #include "MuT/NodeMaterialSkeletalMeshBreak.h"
 
@@ -486,7 +486,7 @@ UE::Mutable::Private::Ptr<UE::Mutable::Private::NodeMaterial> GenerateMutableSou
 		Result = MaterialNode;
 	}
 
-	else if (const UCONodeMaterialSwitch* TypedNodeMaterialSwitch = Cast<UCONodeMaterialSwitch>(Node))
+	else if (const UCONodeSwitch* TypedNodeMaterialSwitch = CastSwitch(Node, UEdGraphSchema_CustomizableObject::PC_Material))
 	{
 		Result = [&]()
 			{
diff --git a/Engine/Plugins/Mutable/Source/CustomizableObjectEditor/Private/MuCOE/GenerateMutableSource/GenerateMutableSourceMesh.cpp b/Engine/Plugins/Mutable/Source/CustomizableObjectEditor/Private/MuCOE/GenerateMutableSource/GenerateMutableSourceMesh.cpp
--- a/Engine/Plugins/Mutable/Source/CustomizableObjectEditor/Private/MuCOE/GenerateMutableSource/GenerateMutableSourceMesh.cpp
+++ b/Engine/Plugins/Mutable/Source/CustomizableObjectEditor/Private/MuCOE/GenerateMutableSource/GenerateMutableSourceMesh.cpp
@@ -35,7 +35,6 @@
 #include "MuCOE/Nodes/CustomizableObjectNodeMeshMorphStackDefinition.h"
 #include "MuCOE/Nodes/CustomizableObjectNodeMeshReshape.h"
 #include "MuCOE/Nodes/CustomizableObjectNodeMeshReshapeCommon.h"
-#include "MuCOE/Nodes/CustomizableObjectNodeMeshSwitch.h"
 #include "MuCOE/Nodes/CustomizableObjectNodeMeshVariation.h"
 #include "MuCOE/Nodes/CustomizableObjectNodeSkeletalMeshParameter.h"
 #include "MuCOE/Nodes/CustomizableObjectNodeSkeletalMesh.h"
@@ -66,6 +65,7 @@
 #include "MuCOE/Nodes/CONodeExternalOperation.h"
 #include "MuCOE/Nodes/CONodeRemoveMesh.h"
 #include "MuCOE/Nodes/CONodeTransformInMesh.h"
+#include "MuCOE/Nodes/CONodeSwitch.h"
 #include "MuT/NodeMeshClipWithMesh.h"
 #include "MuT/NodeMeshExternal.h"
 #include "MuT/NodeMeshRemoveMesh.h"
@@ -2853,7 +2853,7 @@ UE::Mutable::Private::Ptr<UE::Mutable::Private::NodeMesh> GenerateMutableSourceM
 		}
 	}
 
-	else if (const UCustomizableObjectNodeMeshSwitch* TypedNodeMeshSwitch = Cast<UCustomizableObjectNodeMeshSwitch>(Node))
+	else if (const UCONodeSwitch* TypedNodeMeshSwitch = CastSwitch(Node, UEdGraphSchema_CustomizableObject::PC_Mesh))
 	{
 		// Using a lambda so control flow is easier to manage.
 		Result = [&]()
diff --git a/Engine/Plugins/Mutable/Source/CustomizableObjectEditor/Private/MuCOE/GenerateMutableSource/GenerateMutableSourceSurface.cpp b/Engine/Plugins/Mutable/Source/CustomizableObjectEditor/Private/MuCOE/GenerateMutableSource/GenerateMutableSourceSurface.cpp
--- a/Engine/Plugins/Mutable/Source/CustomizableObjectEditor/Private/MuCOE/GenerateMutableSource/GenerateMutableSourceSurface.cpp
+++ b/Engine/Plugins/Mutable/Source/CustomizableObjectEditor/Private/MuCOE/GenerateMutableSource/GenerateMutableSourceSurface.cpp
@@ -2,19 +2,14 @@
 
 #include "MuCOE/GenerateMutableSource/GenerateMutableSourceSurface.h"
 
-#include <MuCOE/Nodes/CONodeSkeletalMeshSection.h>
-#include <MuCOE/Nodes/CONodeSkeletalMeshSectionSwitch.h>
-#include <MuCOE/Nodes/CONodeSkeletalMeshSectionVariation.h>
-
 #include "Engine/SkinnedAssetCommon.h"
 #include "Engine/StaticMesh.h"
 #include "Modules/ModuleManager.h"
-
 #include "GPUSkinVertexFactory.h"
 #include "Rendering/SkeletalMeshLODModel.h"
-
 #include "MuCO/UnrealConversionUtils.h"
 #include "MuCOE/CustomizableObjectCompiler.h"
+#include "MuCOE/EdGraphSchema_CustomizableObject.h"
 #include "MuCOE/GenerateMutableSource/GenerateMutableSourceFloat.h"
 #include "MuCOE/GenerateMutableSource/GenerateMutableSourceGroupProjector.h"
 #include "MuCOE/GenerateMutableSource/GenerateMutableSourceMesh.h"
@@ -26,9 +21,11 @@
 #include "MuCOE/Nodes/CustomizableObjectNodeSkeletalMesh.h"
 #include "MuCOE/Nodes/CustomizableObjectNodeTable.h"
 #include "MuCOE/Nodes/CustomizableObjectNodeTunnel.h"
+#include "MuCOE/Nodes/CONodeSkeletalMeshSection.h"
+#include "MuCOE/Nodes/CONodeSkeletalMeshSectionVariation.h"
+#include "MuCOE/Nodes/CONodeSwitch.h"
 #include "MuT/NodeImageFormat.h"
 #include "MuT/NodeMeshFormat.h"
-#include "MuT/NodeMaterialModify.h"
 #include "MuT/NodeSurfaceSwitch.h"
 #include "MuT/NodeSurfaceVariation.h"
 
@@ -264,7 +261,7 @@ UE::Mutable::Private::Ptr<UE::Mutable::Private::NodeSurface> GenerateMutableSour
 		}
 	}
 
-	else if (const UCONodeSkeletalMeshSectionSwitch* TypedNodeSwitch = Cast<UCONodeSkeletalMeshSectionSwitch>(Node))
+	else if (const UCONodeSwitch* TypedNodeSwitch = CastSwitch(Node, UEdGraphSchema_CustomizableObject::PC_MeshSection))
 	{
 		// Using a lambda so control flow is easier to manage.
 		Result = [&]()
diff --git a/Engine/Plugins/Mutable/Source/CustomizableObjectEditor/Private/MuCOE/GenerateMutableSource/GenerateMutableSourceTransform.cpp b/Engine/Plugins/Mutable/Source/CustomizableObjectEditor/Private/MuCOE/GenerateMutableSource/GenerateMutableSourceTransform.cpp
--- a/Engine/Plugins/Mutable/Source/CustomizableObjectEditor/Private/MuCOE/GenerateMutableSource/GenerateMutableSourceTransform.cpp
+++ b/Engine/Plugins/Mutable/Source/CustomizableObjectEditor/Private/MuCOE/GenerateMutableSource/GenerateMutableSourceTransform.cpp
@@ -2,7 +2,6 @@
 
 #include "GenerateMutableSourceTransform.h"
 
-#include "MuCOE/EdGraphSchema_CustomizableObject.h"
 #include "MuCOE/GenerateMutableSource/GenerateMutableSource.h"
 #include "MuCOE/Nodes/CustomizableObjectNodeTransformConstant.h"
 #include "MuCOE/Nodes/CustomizableObjectNodeTransformParameter.h"
diff --git a/Engine/Plugins/Mutable/Source/CustomizableObjectEditor/Private/MuCOE/GraphTraversal.cpp b/Engine/Plugins/Mutable/Source/CustomizableObjectEditor/Private/MuCOE/GraphTraversal.cpp
--- a/Engine/Plugins/Mutable/Source/CustomizableObjectEditor/Private/MuCOE/GraphTraversal.cpp
+++ b/Engine/Plugins/Mutable/Source/CustomizableObjectEditor/Private/MuCOE/GraphTraversal.cpp
@@ -17,7 +17,6 @@
 #include "MuCOE/Nodes/CustomizableObjectNodeMeshMorphStackApplication.h"
 #include "MuCOE/Nodes/CustomizableObjectNodeMeshMorphStackDefinition.h"
 #include "MuCOE/Nodes/CustomizableObjectNodeMeshReshape.h"
-#include "MuCOE/Nodes/CustomizableObjectNodeMeshSwitch.h"
 #include "MuCOE/Nodes/CustomizableObjectNodeMeshVariation.h"
 #include "MuCOE/Nodes/CustomizableObjectNodeSkeletalMeshParameter.h"
 #include "MuCOE/Nodes/CustomizableObjectNodeObject.h"
@@ -33,6 +32,7 @@
 #include "Nodes/CONodeRemoveMesh.h"
 #include "Nodes/CONodeTransformWithBone.h"
 #include "Nodes/CONodeTransformInMesh.h"
+#include "Nodes/CONodeSwitch.h"
 #include "UObject/UObjectIterator.h"
 
 
@@ -637,7 +637,7 @@ const UEdGraphPin* FindMeshBaseSource(const UEdGraphPin& Pin, const bool bOnlyLo
 		}
 	}
 
-	else if (const UCustomizableObjectNodeMeshSwitch* TypedNodeSwitch = Cast<UCustomizableObjectNodeMeshSwitch>(Node))
+	else if (const UCONodeSwitch* TypedNodeSwitch = CastSwitch(Node, UEdGraphSchema_CustomizableObject::PC_Mesh))
 	{
 		if (const UEdGraphPin* EnumParameterPin = FollowInputPin(*TypedNodeSwitch->SwitchParameterPinReference.Get()))
 		{
diff --git a/Engine/Plugins/Mutable/Source/CustomizableObjectEditor/Private/MuCOE/Nodes/CONodeExternalOperation.cpp b/Engine/Plugins/Mutable/Source/CustomizableObjectEditor/Private/MuCOE/Nodes/CONodeExternalOperation.cpp
--- a/Engine/Plugins/Mutable/Source/CustomizableObjectEditor/Private/MuCOE/Nodes/CONodeExternalOperation.cpp
+++ b/Engine/Plugins/Mutable/Source/CustomizableObjectEditor/Private/MuCOE/Nodes/CONodeExternalOperation.cpp
@@ -11,15 +11,15 @@
 FText UCONodeExternalOperation::GetNodeTitle(ENodeTitleType::Type TitleType) const
 {
 	const UE::Mutable::FExternalOperation* Operation = OperationInstancedStruct.GetPtr<UE::Mutable::FExternalOperation>();
-	const FName Name = Operation ? Operation->GetName() : CachedOperationName;
+	const FText Name = Operation ? OperationInstancedStruct.GetScriptStruct()->GetDisplayNameText() : CachedOperationName;
 	
 	if (TitleType == ENodeTitleType::ListView)
 	{
-		return FText::FromName(Name);
+		return Name;
 	}
 	else
 	{
-		return FText::Format(LOCTEXT("External_NodeTitle", "{0}\nExternal"), FText::FromName(Name));
+		return FText::Format(LOCTEXT("External_NodeTitle", "{0}\nExternal"), Name);
 	}
 }
 
@@ -50,29 +50,23 @@ void UCONodeExternalOperation::AllocateDefaultPins(UCustomizableObjectNodeRemapP
 	// Inputs
 	InputPins.Empty();
 	
-	for (TPair<FName, const UScriptStruct*> Pair : Operation->GetInputs())
+	for (TPair<FText, const UScriptStruct*> Pair : Operation->GetInputs())
 	{
-		UCONodeExtenalPinData* PinData = NewObject<UCONodeExtenalPinData>(this);
-		PinData->Name = Pair.Key;
-		
 		const FName Category = UEdGraphSchema_CustomizableObject::GetPinCategoryName(Pair.Value);
 
-		UEdGraphPin* Pin = CustomCreatePin(EGPD_Input, Category, Pair.Key, FText::FromName(Pair.Key), PinData);
+		UEdGraphPin* Pin = CustomCreatePin(EGPD_Input, Category, FName(Pair.Key.ToString()), Pair.Key);
 		InputPins.Add(Pin);
 	}
 
 	// Output
 	{
-		const TPair<FName, const UScriptStruct*> Pair = Operation->GetOutput();
+		const TPair<FText, const UScriptStruct*> Pair = Operation->GetOutput();
 
-		UCONodeExtenalPinData* PinData = NewObject<UCONodeExtenalPinData>(this);
-		PinData->Name = Pair.Key;
-
-		const FName Category = UEdGraphSchema_CustomizableObject::GetPinCategoryName(Pair.Value);
-		OutputPin = CustomCreatePin(EGPD_Output, Category, Pair.Key, FText::FromName(Pair.Key), PinData);
+		const FName Category = UEdGraphSchema_CustomizableObject::GetPinCategoryName(Pair.Value);	
+		OutputPin = CustomCreatePin(EGPD_Output, Category, FName(Pair.Key.ToString()), Pair.Key);
 	}
 
-	CachedOperationName = Operation->GetName();
+	CachedOperationName = OperationInstancedStruct.GetScriptStruct()->GetDisplayNameText();
 }
 
 
diff --git a/Engine/Plugins/Mutable/Source/CustomizableObjectEditor/Private/MuCOE/Nodes/CONodeExternalOperation.h b/Engine/Plugins/Mutable/Source/CustomizableObjectEditor/Private/MuCOE/Nodes/CONodeExternalOperation.h
--- a/Engine/Plugins/Mutable/Source/CustomizableObjectEditor/Private/MuCOE/Nodes/CONodeExternalOperation.h
+++ b/Engine/Plugins/Mutable/Source/CustomizableObjectEditor/Private/MuCOE/Nodes/CONodeExternalOperation.h
@@ -7,17 +7,6 @@
 #include "CONodeExternalOperation.generated.h"
 
 
-UCLASS()
-class UCONodeExtenalPinData : public UCustomizableObjectNodePinData
-{
-	GENERATED_BODY()
-
-public:
-	UPROPERTY()
-	FName Name;
-};
-
-
 UCLASS()
 class UCONodeExternalOperation : public UCustomizableObjectNode
 {
@@ -43,5 +32,5 @@ public:
 	
 	/** In case of the operation not being loaded, the node uses this cached name so we can at least show the user which operation it was. */
 	UPROPERTY()
-	FName CachedOperationName;
+	FText CachedOperationName;
 };
diff --git a/Engine/Plugins/Mutable/Source/CustomizableObjectEditor/Private/MuCOE/Nodes/CONodeExternalTypeParameter.cpp b/Engine/Plugins/Mutable/Source/CustomizableObjectEditor/Private/MuCOE/Nodes/CONodeExternalTypeParameter.cpp
--- a/Engine/Plugins/Mutable/Source/CustomizableObjectEditor/Private/MuCOE/Nodes/CONodeExternalTypeParameter.cpp
+++ b/Engine/Plugins/Mutable/Source/CustomizableObjectEditor/Private/MuCOE/Nodes/CONodeExternalTypeParameter.cpp
@@ -37,7 +37,7 @@ FText UCONodeExternalTypeParameter::GetCategoryFriendlyName() const
 {
 	if (DefaultValue.IsValid())
 	{
-		CachedFriendlyName = UEdGraphSchema_CustomizableObject::GetPinCategoryFriendlyName(DefaultValue.GetScriptStruct());
+		CachedFriendlyName = UEdGraphSchema_CustomizableObject::GetPinCategoryFriendlyName(GetCategory());
 	}
 
 	return CachedFriendlyName;
diff --git a/Engine/Plugins/Mutable/Source/CustomizableObjectEditor/Private/MuCOE/Nodes/CONodeMaterialSwitch.cpp b/Engine/Plugins/Mutable/Source/CustomizableObjectEditor/Private/MuCOE/Nodes/CONodeMaterialSwitch.cpp
--- a/Engine/Plugins/Mutable/Source/CustomizableObjectEditor/Private/MuCOE/Nodes/CONodeMaterialSwitch.cpp
+++ b/Engine/Plugins/Mutable/Source/CustomizableObjectEditor/Private/MuCOE/Nodes/CONodeMaterialSwitch.cpp
@@ -1,15 +1,20 @@
 // Copyright Epic Games, Inc. All Rights Reserved.
 
 #include "MuCOE/Nodes/CONodeMaterialSwitch.h"
+
+#include "MuCO/CustomizableObjectCustomVersion.h"
 #include "MuCOE/EdGraphSchema_CustomizableObject.h"
+
 #include UE_INLINE_GENERATED_CPP_BY_NAME(CONodeMaterialSwitch)
 
 #define LOCTEXT_NAMESPACE "CustomizableObjectEditor"
 
 
-FName UCONodeMaterialSwitch::GetCategory() const
+void UCONodeMaterialSwitch::BackwardsCompatibleFixup(int32 CustomizableObjectCustomVersion)
 {
-	return UEdGraphSchema_CustomizableObject::PC_Material;
+	PinType = UEdGraphSchema_CustomizableObject::PC_Material;
+	
+	Super::BackwardsCompatibleFixup(CustomizableObjectCustomVersion);
 }
 
 
diff --git a/Engine/Plugins/Mutable/Source/CustomizableObjectEditor/Private/MuCOE/Nodes/CONodeMaterialSwitch.h b/Engine/Plugins/Mutable/Source/CustomizableObjectEditor/Private/MuCOE/Nodes/CONodeMaterialSwitch.h
--- a/Engine/Plugins/Mutable/Source/CustomizableObjectEditor/Private/MuCOE/Nodes/CONodeMaterialSwitch.h
+++ b/Engine/Plugins/Mutable/Source/CustomizableObjectEditor/Private/MuCOE/Nodes/CONodeMaterialSwitch.h
@@ -2,20 +2,23 @@
 
 #pragma once
 
-#include "MuCOE/Nodes/CustomizableObjectNodeSwitchBase.h"
+#include "MuCOE/Nodes/CONodeSwitch.h"
+
 #include "CONodeMaterialSwitch.generated.h"
 
 #define UE_API CUSTOMIZABLEOBJECTEDITOR_API
 
 
-UCLASS(MinimalAPI)
-class UCONodeMaterialSwitch : public UCustomizableObjectNodeSwitchBase
+// DEPRECATED
+UCLASS(NotPlaceable, HideDropdown)
+class UCONodeMaterialSwitch : public UCONodeSwitch
 {
 public:
 	GENERATED_BODY()
 
-	// UCustomizableObjectNodeSwitchBase interface
-	UE_API virtual FName GetCategory() const override;
+	// UCustomizableObjectNodeSwitch interface
+	virtual void BackwardsCompatibleFixup(int32 CustomizableObjectCustomVersion) override;
 };
 
+
 #undef UE_API
diff --git a/Engine/Plugins/Mutable/Source/CustomizableObjectEditor/Private/MuCOE/Nodes/CONodeSkeletalMeshSectionSwitch.cpp b/Engine/Plugins/Mutable/Source/CustomizableObjectEditor/Private/MuCOE/Nodes/CONodeSkeletalMeshSectionSwitch.cpp
--- a/Engine/Plugins/Mutable/Source/CustomizableObjectEditor/Private/MuCOE/Nodes/CONodeSkeletalMeshSectionSwitch.cpp
+++ b/Engine/Plugins/Mutable/Source/CustomizableObjectEditor/Private/MuCOE/Nodes/CONodeSkeletalMeshSectionSwitch.cpp
@@ -10,13 +10,10 @@
 #define LOCTEXT_NAMESPACE "CustomizableObjectEditor"
 
 
-FName UCONodeSkeletalMeshSectionSwitch::GetCategory() const
-{
-	return UEdGraphSchema_CustomizableObject::PC_MeshSection;
-}
-
 void UCONodeSkeletalMeshSectionSwitch::BackwardsCompatibleFixup(int32 CustomizableObjectCustomVersion)
 {
+	PinType = UEdGraphSchema_CustomizableObject::PC_MeshSection;
+
 	Super::BackwardsCompatibleFixup(CustomizableObjectCustomVersion);
 
 	if (CustomizableObjectCustomVersion == FCustomizableObjectCustomVersion::FixMaterialPinsRename)
diff --git a/Engine/Plugins/Mutable/Source/CustomizableObjectEditor/Private/MuCOE/Nodes/CONodeSkeletalMeshSectionSwitch.h b/Engine/Plugins/Mutable/Source/CustomizableObjectEditor/Private/MuCOE/Nodes/CONodeSkeletalMeshSectionSwitch.h
--- a/Engine/Plugins/Mutable/Source/CustomizableObjectEditor/Private/MuCOE/Nodes/CONodeSkeletalMeshSectionSwitch.h
+++ b/Engine/Plugins/Mutable/Source/CustomizableObjectEditor/Private/MuCOE/Nodes/CONodeSkeletalMeshSectionSwitch.h
@@ -2,24 +2,23 @@
 
 #pragma once
 
-#include "MuCOE/Nodes/CustomizableObjectNodeSwitchBase.h"
+#include "MuCOE/Nodes/CONodeSwitch.h"
 
 #include "CONodeSkeletalMeshSectionSwitch.generated.h"
 
 #define UE_API CUSTOMIZABLEOBJECTEDITOR_API
 
 
-UCLASS(MinimalAPI)
-class UCONodeSkeletalMeshSectionSwitch : public UCustomizableObjectNodeSwitchBase
+// DEPRECATED
+UCLASS(NotPlaceable, HideDropdown)
+class UCONodeSkeletalMeshSectionSwitch : public UCONodeSwitch
 {
 public:
 	GENERATED_BODY()
 
 	// UCustomizableObjectNode interface
-	UE_API virtual void BackwardsCompatibleFixup(int32 CustomizableObjectCustomVersion) override;
-
-	// UCustomizableObjectNodeSwitchBase interface
-	UE_API virtual FName GetCategory() const override;
+	virtual void BackwardsCompatibleFixup(int32 CustomizableObjectCustomVersion) override;
 };
 
+
 #undef UE_API
diff --git a/Engine/Plugins/Mutable/Source/CustomizableObjectEditor/Private/MuCOE/Nodes/CONodeSwitch.cpp b/Engine/Plugins/Mutable/Source/CustomizableObjectEditor/Private/MuCOE/Nodes/CONodeSwitch.cpp
new file mode 100644
--- /dev/null
+++ b/Engine/Plugins/Mutable/Source/CustomizableObjectEditor/Private/MuCOE/Nodes/CONodeSwitch.cpp
@@ -0,0 +1,180 @@
+// Copyright Epic Games, Inc. All Rights Reserved.
+
+#include "MuCOE/Nodes/CONodeSwitch.h"
+
+#include "MuCO/CustomizableObjectCustomVersion.h"
+#include "MuCOE/EdGraphSchema_CustomizableObject.h"
+#include "MuCOE/GraphTraversal.h"
+#include "MuCOE/Nodes/CustomizableObjectNodeEnumParameter.h"
+#include "MuT/NodeScalarEnumParameter.h"
+
+#include UE_INLINE_GENERATED_CPP_BY_NAME(CONodeSwitch)
+
+class UCustomizableObjectNodeRemapPins;
+
+#define LOCTEXT_NAMESPACE "CustomizableObjectEditor"
+
+
+
+void UCONodeSwitch::PinConnectionListChanged(UEdGraphPin* Pin)
+{
+	if (Pin == SwitchParameterPinReference.Get())
+	{
+		Super::ReconstructNode();
+	}
+
+	Super::PinConnectionListChanged(Pin);
+}
+
+
+void UCONodeSwitch::AllocateDefaultPins(UCustomizableObjectNodeRemapPins* RemapPins)
+{
+	const UEdGraphSchema_CustomizableObject* Schema = GetDefault<UEdGraphSchema_CustomizableObject>();
+	
+	CustomCreatePin(EGPD_Output, GetCategory(), FName(UEdGraphSchema_CustomizableObject::GetPinCategoryName(GetCategory()).ToString()), UEdGraphSchema_CustomizableObject::GetPinCategoryFriendlyName(GetCategory()));
+
+	UEdGraphPin* OldEnumPin = SwitchParameterPinReference.Get();
+	SwitchParameterPinReference = CustomCreatePin(EGPD_Input, Schema->PC_Enum, TEXT("Switch Parameter"), LOCTEXT("SwitchParameter", "Switch Parameter"));
+
+	SwitchPins.Empty();	
+	if (OldEnumPin) // May be null the first time the node is created.
+	{
+		// This will fail if the node is connected to a macro node which will then prevent the generation of the input option pins
+		// todo : UE-294591
+		if (const UEdGraphPin* LinkedPin = FollowInputPin(*OldEnumPin))
+		{
+			if (UCustomizableObjectNodeEnumParameter* EnumNode = Cast<UCustomizableObjectNodeEnumParameter>(LinkedPin->GetOwningNode()))
+			{
+				for (int32 EnumValueIndex = 0; EnumValueIndex < EnumNode->Values.Num(); ++EnumValueIndex)
+				{
+					const FString OptionName = EnumNode->Values[EnumValueIndex].Name;
+					SwitchPins.Add(CustomCreatePin(EGPD_Input, GetCategory(), FName(*OptionName), FText::FromString(OptionName)));
+				}
+			}
+		}
+	}
+}
+
+
+FName UCONodeSwitch::GetCategory() const
+{
+	return PinType;
+}
+
+
+FText UCONodeSwitch::GetNodeTitle(ENodeTitleType::Type TitleType) const
+{
+	return FText::Format(LOCTEXT("Switch_Title", "{0} Switch"), UEdGraphSchema_CustomizableObject::GetPinCategoryFriendlyName(GetCategory()));
+}
+
+
+FLinearColor UCONodeSwitch::GetNodeTitleColor() const
+{
+	return UEdGraphSchema_CustomizableObject::GetPinTypeColor(GetCategory());
+}
+
+
+FText UCONodeSwitch::GetTooltipText() const
+{
+	return LOCTEXT("Switch_Tooltip", "Change the resulting value depending on what is currently chosen among a predefined amount of sources.");
+}
+
+
+void UCONodeSwitch::BackwardsCompatibleFixup(int32 CustomizableObjectCustomVersion)
+{
+	Super::BackwardsCompatibleFixup(CustomizableObjectCustomVersion);
+
+	if (CustomizableObjectCustomVersion == FCustomizableObjectCustomVersion::BugPinsSwitch)
+	{
+		SwitchParameterPinReference = FindPin(TEXT("Switch Parameter"));	
+	}
+	
+	else if (CustomizableObjectCustomVersion == FCustomizableObjectCustomVersion::ChangedSwitchNodesInputPinsFriendlyNames)
+	{
+		const FString CategoryFriendlyName = UEdGraphSchema_CustomizableObject::GetPinCategoryFriendlyName(GetCategory()).ToString();
+		const FString PinPrefix = UEdGraphSchema_CustomizableObject::GetPinCategoryName(GetCategory()).ToString() + " ";
+
+		int32 NumOptions = 0;
+		for (const UEdGraphPin* Pin : GetAllNonOrphanPins())
+		{
+			if (Pin->GetName().StartsWith(PinPrefix))
+			{
+				NumOptions++;
+			}
+		}
+		
+		for (int32 ElementIndex = 0; ElementIndex < NumOptions; ++ElementIndex)
+		{
+			UEdGraphPin* ElementPin = FindPin(PinPrefix + FString::FromInt(ElementIndex) + TEXT(" "));
+
+			if (ElementPin->PinFriendlyName.ToString().StartsWith(TEXT("Material")))
+			{
+				ElementPin->PinFriendlyName = FText::FromString(CategoryFriendlyName + TEXT(" ") + FString::FromInt(ElementIndex));
+			}
+		}
+	}
+	
+	else if (CustomizableObjectCustomVersion == FCustomizableObjectCustomVersion::PassthroughParameter)
+	{
+		const FString PinPrefix = UEdGraphSchema_CustomizableObject::GetPinCategoryName(GetCategory()).ToString() + " ";
+
+		int32 NumOptions = 0;
+		for (const UEdGraphPin* Pin : GetAllNonOrphanPins())
+		{
+			if (Pin->GetName().StartsWith(PinPrefix))
+			{
+				NumOptions++;
+			}
+		}
+		
+		TArray<FString> OptionNames;
+		if (const UEdGraphPin* LinkedPin = FollowInputPin(*SwitchParameterPinReference.Get()))
+		{
+			if (UCustomizableObjectNodeEnumParameter* EnumNode = Cast<UCustomizableObjectNodeEnumParameter>(LinkedPin->GetOwningNode()))
+			{
+				for (int32 ValueIndex = 0; ValueIndex < EnumNode->Values.Num(); ++ValueIndex)
+				{
+					OptionNames.Add(EnumNode->Values[ValueIndex].Name);
+				}
+			}
+		}
+		
+		for (int32 ElementIndex = 0; ElementIndex < NumOptions; ++ElementIndex)
+		{
+			UEdGraphPin* ElementPin = FindPin(PinPrefix + FString::FromInt(ElementIndex) + TEXT(" "));
+			if (OptionNames.IsValidIndex(ElementIndex))
+			{
+				ElementPin->PinName = FName(OptionNames[ElementIndex]);
+				ElementPin->PinFriendlyName = FText::FromName(ElementPin->PinName);
+				SwitchPins.Add(ElementPin);
+			}
+			else
+			{
+				ElementPin->bOrphanedPin = true;
+			}
+		}
+	}
+}
+
+
+void UCONodeSwitch::PostBackwardsCompatibleFixup()
+{
+	Super::PostBackwardsCompatibleFixup();
+	Super::ReconstructNode();
+}
+
+
+FString UCONodeSwitch::GetOutputPinName() const
+{
+	return FString();
+}
+
+
+const UCONodeSwitch* CastSwitch(const UEdGraphNode* Node, FName Type)
+{
+	const UCONodeSwitch* NodeSwitch = Cast<const UCONodeSwitch>(Node);
+	return NodeSwitch && NodeSwitch->PinType == Type ? NodeSwitch : nullptr;
+}
+
+
+#undef LOCTEXT_NAMESPACE
diff --git a/Engine/Plugins/Mutable/Source/CustomizableObjectEditor/Private/MuCOE/Nodes/CONodeSwitch.h b/Engine/Plugins/Mutable/Source/CustomizableObjectEditor/Private/MuCOE/Nodes/CONodeSwitch.h
new file mode 100644
--- /dev/null
+++ b/Engine/Plugins/Mutable/Source/CustomizableObjectEditor/Private/MuCOE/Nodes/CONodeSwitch.h
@@ -0,0 +1,65 @@
+// Copyright Epic Games, Inc. All Rights Reserved.
+
+#pragma once
+
+#include "MuCOE/Nodes/CustomizableObjectNode.h"
+
+#include "CONodeSwitch.generated.h"
+
+#define UE_API CUSTOMIZABLEOBJECTEDITOR_API
+
+namespace ENodeTitleType { enum Type : int; }
+
+class FArchive;
+class UCustomizableObjectNodeRemapPins;
+class UObject;
+struct FFrame;
+
+namespace UE::Mutable::Private
+{
+	class NodeScalarEnumParameter;
+}
+
+UCLASS()
+class UCONodeSwitch : public UCustomizableObjectNode
+{
+public:
+	GENERATED_BODY()
+	
+	// Begin EdGraphNode interface
+	virtual FText GetNodeTitle(ENodeTitleType::Type TitleType) const override;
+	virtual FLinearColor GetNodeTitleColor() const override;
+	virtual FText GetTooltipText() const override;
+	virtual void PinConnectionListChanged(UEdGraphPin* Pin) override;
+	
+	// UCustomizableObjectNode interface
+	virtual void BackwardsCompatibleFixup(int32 CustomizableObjectCustomVersion) override;
+	virtual void PostBackwardsCompatibleFixup() override;
+	virtual void AllocateDefaultPins(UCustomizableObjectNodeRemapPins* RemapPins) override;
+
+	/** Get the output pin category. Override. */
+	virtual FName GetCategory() const;
+
+	/** Get the ouput pin name. Override. */
+	virtual FString GetOutputPinName() const;
+	
+private:
+	
+	/** Last NodeEnumParameter connected. Used to remove the callback once disconnected. */
+	TWeakObjectPtr<UCustomizableObjectNode> LastNodeEnumParameterConnected;
+
+public:
+	UPROPERTY()
+	FEdGraphPinReference SwitchParameterPinReference;
+	
+	UPROPERTY()
+	TArray<FEdGraphPinReference> SwitchPins;
+
+	UPROPERTY()
+	FName PinType;
+};
+
+
+const UCONodeSwitch* CastSwitch(const UEdGraphNode* Node, FName Type);
+
+#undef UE_API
diff --git a/Engine/Plugins/Mutable/Source/CustomizableObjectEditor/Private/MuCOE/Nodes/CustomizableObjectNodeColorSwitch.cpp b/Engine/Plugins/Mutable/Source/CustomizableObjectEditor/Private/MuCOE/Nodes/CustomizableObjectNodeColorSwitch.cpp
--- a/Engine/Plugins/Mutable/Source/CustomizableObjectEditor/Private/MuCOE/Nodes/CustomizableObjectNodeColorSwitch.cpp
+++ b/Engine/Plugins/Mutable/Source/CustomizableObjectEditor/Private/MuCOE/Nodes/CustomizableObjectNodeColorSwitch.cpp
@@ -7,14 +7,10 @@
 
 #include UE_INLINE_GENERATED_CPP_BY_NAME(CustomizableObjectNodeColorSwitch)
 
-#define LOCTEXT_NAMESPACE "CustomizableObjectEditor"
 
-
-FName UCustomizableObjectNodeColorSwitch::GetCategory() const
+void UCustomizableObjectNodeColorSwitch::BackwardsCompatibleFixup(int32 CustomizableObjectCustomVersion)
 {
-	return UEdGraphSchema_CustomizableObject::PC_Color;
-}
-
-
-#undef LOCTEXT_NAMESPACE
+	PinType = UEdGraphSchema_CustomizableObject::PC_Color;
 
+	Super::BackwardsCompatibleFixup(CustomizableObjectCustomVersion);
+}
diff --git a/Engine/Plugins/Mutable/Source/CustomizableObjectEditor/Private/MuCOE/Nodes/CustomizableObjectNodeColorSwitch.h b/Engine/Plugins/Mutable/Source/CustomizableObjectEditor/Private/MuCOE/Nodes/CustomizableObjectNodeColorSwitch.h
--- a/Engine/Plugins/Mutable/Source/CustomizableObjectEditor/Private/MuCOE/Nodes/CustomizableObjectNodeColorSwitch.h
+++ b/Engine/Plugins/Mutable/Source/CustomizableObjectEditor/Private/MuCOE/Nodes/CustomizableObjectNodeColorSwitch.h
@@ -2,21 +2,23 @@
 
 #pragma once
 
-#include "MuCOE/Nodes/CustomizableObjectNodeSwitchBase.h"
+#include "MuCOE/Nodes/CONodeSwitch.h"
 
 #include "CustomizableObjectNodeColorSwitch.generated.h"
 
 #define UE_API CUSTOMIZABLEOBJECTEDITOR_API
 
 
-UCLASS(MinimalAPI)
-class UCustomizableObjectNodeColorSwitch : public UCustomizableObjectNodeSwitchBase
+// DEPRECATED
+UCLASS(NotPlaceable, HideDropdown)
+class UCustomizableObjectNodeColorSwitch : public UCONodeSwitch
 {
 public:
 	GENERATED_BODY()
 
-	// UCustomizableObjectNodeSwitchBase interface
-	UE_API virtual FName GetCategory() const override;
+	// UCustomizableObjectNodeSwitch interface
+	virtual void BackwardsCompatibleFixup(int32 CustomizableObjectCustomVersion) override;
 };
 
+
 #undef UE_API
diff --git a/Engine/Plugins/Mutable/Source/CustomizableObjectEditor/Private/MuCOE/Nodes/CustomizableObjectNodeComponentSwitch.cpp b/Engine/Plugins/Mutable/Source/CustomizableObjectEditor/Private/MuCOE/Nodes/CustomizableObjectNodeComponentSwitch.cpp
--- a/Engine/Plugins/Mutable/Source/CustomizableObjectEditor/Private/MuCOE/Nodes/CustomizableObjectNodeComponentSwitch.cpp
+++ b/Engine/Plugins/Mutable/Source/CustomizableObjectEditor/Private/MuCOE/Nodes/CustomizableObjectNodeComponentSwitch.cpp
@@ -2,6 +2,7 @@
 
 #include "MuCOE/Nodes/CustomizableObjectNodeComponentSwitch.h"
 
+#include "MuCO/CustomizableObjectCustomVersion.h"
 #include "MuCOE/EdGraphSchema_CustomizableObject.h"
 
 #include UE_INLINE_GENERATED_CPP_BY_NAME(CustomizableObjectNodeComponentSwitch)
@@ -9,9 +10,11 @@
 #define LOCTEXT_NAMESPACE "CustomizableObjectEditor"
 
 
-FName UCustomizableObjectNodeComponentSwitch::GetCategory() const
+void UCustomizableObjectNodeComponentSwitch::BackwardsCompatibleFixup(int32 CustomizableObjectCustomVersion)
 {
-	return UEdGraphSchema_CustomizableObject::PC_Component;
+	PinType = UEdGraphSchema_CustomizableObject::PC_Component;
+	
+	Super::BackwardsCompatibleFixup(CustomizableObjectCustomVersion);
 }
 
 
diff --git a/Engine/Plugins/Mutable/Source/CustomizableObjectEditor/Private/MuCOE/Nodes/CustomizableObjectNodeComponentSwitch.h b/Engine/Plugins/Mutable/Source/CustomizableObjectEditor/Private/MuCOE/Nodes/CustomizableObjectNodeComponentSwitch.h
--- a/Engine/Plugins/Mutable/Source/CustomizableObjectEditor/Private/MuCOE/Nodes/CustomizableObjectNodeComponentSwitch.h
+++ b/Engine/Plugins/Mutable/Source/CustomizableObjectEditor/Private/MuCOE/Nodes/CustomizableObjectNodeComponentSwitch.h
@@ -2,21 +2,22 @@
 
 #pragma once
 
-#include "MuCOE/Nodes/CustomizableObjectNodeSwitchBase.h"
+#include "MuCOE/Nodes/CONodeSwitch.h"
 
 #include "CustomizableObjectNodeComponentSwitch.generated.h"
 
 #define UE_API CUSTOMIZABLEOBJECTEDITOR_API
 
 
-UCLASS(MinimalAPI)
-class UCustomizableObjectNodeComponentSwitch : public UCustomizableObjectNodeSwitchBase
+// DEPRECATED
+UCLASS(NotPlaceable, HideDropdown)
+class UCustomizableObjectNodeComponentSwitch : public UCONodeSwitch
 {
 public:
 	GENERATED_BODY()
 	
-	// UCustomizableObjectNodeSwitchBase interface
-	UE_API virtual FName GetCategory() const override;
+	// UCustomizableObjectNodeSwitch interface
+	virtual void BackwardsCompatibleFixup(int32 CustomizableObjectCustomVersion) override;
 };
 
 #undef UE_API
diff --git a/Engine/Plugins/Mutable/Source/CustomizableObjectEditor/Private/MuCOE/Nodes/CustomizableObjectNodeExposePin.cpp b/Engine/Plugins/Mutable/Source/CustomizableObjectEditor/Private/MuCOE/Nodes/CustomizableObjectNodeExposePin.cpp
--- a/Engine/Plugins/Mutable/Source/CustomizableObjectEditor/Private/MuCOE/Nodes/CustomizableObjectNodeExposePin.cpp
+++ b/Engine/Plugins/Mutable/Source/CustomizableObjectEditor/Private/MuCOE/Nodes/CustomizableObjectNodeExposePin.cpp
@@ -53,7 +53,7 @@ FText UCustomizableObjectNodeExposePin::GetNodeTitle(ENodeTitleType::Type TitleT
 
 	if (TitleType == ENodeTitleType::ListView)
 	{
-		NodeTitle = LOCTEXT("Expose_Pin_Title_ListView", "Export {PinType} Pin");
+		NodeTitle = LOCTEXT("Expose_Pin_Title_ListView", "{PinType} Export");
 	}
 	else if (TitleType == ENodeTitleType::EditableTitle)
 	{
@@ -61,7 +61,7 @@ FText UCustomizableObjectNodeExposePin::GetNodeTitle(ENodeTitleType::Type TitleT
 	}
 	else
 	{
-		NodeTitle = LOCTEXT("Expose_Pin_Title", "{NodeName}\nExport {PinType} Pin");
+		NodeTitle = LOCTEXT("Expose_Pin_Title", "{NodeName}\n{PinType} Export");
 	}
 
 	return FText::Format(NodeTitle, Args);
diff --git a/Engine/Plugins/Mutable/Source/CustomizableObjectEditor/Private/MuCOE/Nodes/CustomizableObjectNodeExtensionDataSwitch.cpp b/Engine/Plugins/Mutable/Source/CustomizableObjectEditor/Private/MuCOE/Nodes/CustomizableObjectNodeExtensionDataSwitch.cpp
deleted file mode 100644
diff --git a/Engine/Plugins/Mutable/Source/CustomizableObjectEditor/Private/MuCOE/Nodes/CustomizableObjectNodeExtensionDataSwitch.h b/Engine/Plugins/Mutable/Source/CustomizableObjectEditor/Private/MuCOE/Nodes/CustomizableObjectNodeExtensionDataSwitch.h
deleted file mode 100644
diff --git a/Engine/Plugins/Mutable/Source/CustomizableObjectEditor/Private/MuCOE/Nodes/CustomizableObjectNodeExternalPin.cpp b/Engine/Plugins/Mutable/Source/CustomizableObjectEditor/Private/MuCOE/Nodes/CustomizableObjectNodeExternalPin.cpp
--- a/Engine/Plugins/Mutable/Source/CustomizableObjectEditor/Private/MuCOE/Nodes/CustomizableObjectNodeExternalPin.cpp
+++ b/Engine/Plugins/Mutable/Source/CustomizableObjectEditor/Private/MuCOE/Nodes/CustomizableObjectNodeExternalPin.cpp
@@ -247,17 +247,17 @@ FText UCustomizableObjectNodeExternalPin::GetNodeTitle(ENodeTitleType::Type Titl
 
 	if (TitleType == ENodeTitleType::ListView || !ExternalObject)
 	{
-		return FText::Format(LOCTEXT("External_Pin_Title", "Import {0} Pin"), PinTypeName);
+		return FText::Format(LOCTEXT("External_Pin_Title", "{0} Import"), PinTypeName);
 	}
 	else
 	{
 		if (const UCustomizableObjectNodeExposePin* ExportNode = GetNodeExposePin())
 		{
-			return FText::Format(LOCTEXT("External_Pin_ExportNode_Title", "{0}\nImport {1} Pin"), FText::FromString(*ExportNode->GetNodeName()), PinTypeName);
+			return FText::Format(LOCTEXT("External_Pin_ExportNode_Title", "{0}\n{1} Import"), FText::FromString(*ExportNode->GetNodeName()), PinTypeName);
 		}
 		else
 		{
-			return FText::Format(LOCTEXT("External_Pin_ExternalObject_Title", "{0}\nImport {1} Pin"), FText::FromString(ExternalObject->GetName()), PinTypeName);
+			return FText::Format(LOCTEXT("External_Pin_ExternalObject_Title", "{0}\n{1} Import"), FText::FromString(ExternalObject->GetName()), PinTypeName);
 		}
 	}
 }
diff --git a/Engine/Plugins/Mutable/Source/CustomizableObjectEditor/Private/MuCOE/Nodes/CustomizableObjectNodeFloatSwitch.cpp b/Engine/Plugins/Mutable/Source/CustomizableObjectEditor/Private/MuCOE/Nodes/CustomizableObjectNodeFloatSwitch.cpp
--- a/Engine/Plugins/Mutable/Source/CustomizableObjectEditor/Private/MuCOE/Nodes/CustomizableObjectNodeFloatSwitch.cpp
+++ b/Engine/Plugins/Mutable/Source/CustomizableObjectEditor/Private/MuCOE/Nodes/CustomizableObjectNodeFloatSwitch.cpp
@@ -10,9 +10,12 @@
 #define LOCTEXT_NAMESPACE "CustomizableObjectEditor"
 
 
-FName UCustomizableObjectNodeFloatSwitch::GetCategory() const
+void UCustomizableObjectNodeFloatSwitch::BackwardsCompatibleFixup(int32 CustomizableObjectCustomVersion)
 {
-	return UEdGraphSchema_CustomizableObject::PC_Float;
+	PinType = UEdGraphSchema_CustomizableObject::PC_Float;
+	
+	Super::BackwardsCompatibleFixup(CustomizableObjectCustomVersion);
+
 }
 
 
diff --git a/Engine/Plugins/Mutable/Source/CustomizableObjectEditor/Private/MuCOE/Nodes/CustomizableObjectNodeFloatSwitch.h b/Engine/Plugins/Mutable/Source/CustomizableObjectEditor/Private/MuCOE/Nodes/CustomizableObjectNodeFloatSwitch.h
--- a/Engine/Plugins/Mutable/Source/CustomizableObjectEditor/Private/MuCOE/Nodes/CustomizableObjectNodeFloatSwitch.h
+++ b/Engine/Plugins/Mutable/Source/CustomizableObjectEditor/Private/MuCOE/Nodes/CustomizableObjectNodeFloatSwitch.h
@@ -2,21 +2,22 @@
 
 #pragma once
 
-#include "MuCOE/Nodes/CustomizableObjectNodeSwitchBase.h"
+#include "MuCOE/Nodes/CONodeSwitch.h"
 
 #include "CustomizableObjectNodeFloatSwitch.generated.h"
 
 #define UE_API CUSTOMIZABLEOBJECTEDITOR_API
 
 
-UCLASS(MinimalAPI)
-class UCustomizableObjectNodeFloatSwitch : public UCustomizableObjectNodeSwitchBase
+// DEPRECATED
+UCLASS(NotPlaceable, HideDropdown)
+class UCustomizableObjectNodeFloatSwitch : public UCONodeSwitch
 {
 public:
 	GENERATED_BODY()
 
-	// UCustomizableObjectNodeSwitchBase interface
-	UE_API virtual FName GetCategory() const override;
+	// UCustomizableObjectNodeSwitch interface
+	virtual void BackwardsCompatibleFixup(int32 CustomizableObjectCustomVersion) override;
 };
 
 #undef UE_API
diff --git a/Engine/Plugins/Mutable/Source/CustomizableObjectEditor/Private/MuCOE/Nodes/CustomizableObjectNodeMeshSwitch.cpp b/Engine/Plugins/Mutable/Source/CustomizableObjectEditor/Private/MuCOE/Nodes/CustomizableObjectNodeMeshSwitch.cpp
--- a/Engine/Plugins/Mutable/Source/CustomizableObjectEditor/Private/MuCOE/Nodes/CustomizableObjectNodeMeshSwitch.cpp
+++ b/Engine/Plugins/Mutable/Source/CustomizableObjectEditor/Private/MuCOE/Nodes/CustomizableObjectNodeMeshSwitch.cpp
@@ -10,9 +10,11 @@
 #define LOCTEXT_NAMESPACE "CustomizableObjectEditor"
 
 
-FName UCustomizableObjectNodeMeshSwitch::GetCategory() const
+void UCustomizableObjectNodeMeshSwitch::BackwardsCompatibleFixup(int32 CustomizableObjectCustomVersion)
 {
-	return UEdGraphSchema_CustomizableObject::PC_Mesh;
+	PinType = UEdGraphSchema_CustomizableObject::PC_Mesh;
+	
+	Super::BackwardsCompatibleFixup(CustomizableObjectCustomVersion);
 }
 
 
diff --git a/Engine/Plugins/Mutable/Source/CustomizableObjectEditor/Private/MuCOE/Nodes/CustomizableObjectNodeMeshSwitch.h b/Engine/Plugins/Mutable/Source/CustomizableObjectEditor/Private/MuCOE/Nodes/CustomizableObjectNodeMeshSwitch.h
--- a/Engine/Plugins/Mutable/Source/CustomizableObjectEditor/Private/MuCOE/Nodes/CustomizableObjectNodeMeshSwitch.h
+++ b/Engine/Plugins/Mutable/Source/CustomizableObjectEditor/Private/MuCOE/Nodes/CustomizableObjectNodeMeshSwitch.h
@@ -2,21 +2,22 @@
 
 #pragma once
 
-#include "MuCOE/Nodes/CustomizableObjectNodeSwitchBase.h"
+#include "MuCOE/Nodes/CONodeSwitch.h"
 
 #include "CustomizableObjectNodeMeshSwitch.generated.h"
 
 #define UE_API CUSTOMIZABLEOBJECTEDITOR_API
 
 
-UCLASS(MinimalAPI)
-class UCustomizableObjectNodeMeshSwitch : public UCustomizableObjectNodeSwitchBase
+// DEPRECATED
+UCLASS(NotPlaceable, HideDropdown)
+class UCustomizableObjectNodeMeshSwitch : public UCONodeSwitch
 {
 public:
 	GENERATED_BODY()
 
-	// UCustomizableObjectNodeSwitchBase interface
-	UE_API virtual FName GetCategory() const override;
+	// UCustomizableObjectNodeSwitch interface
+	virtual void BackwardsCompatibleFixup(int32 CustomizableObjectCustomVersion) override;
 };
 
 #undef UE_API
diff --git a/Engine/Plugins/Mutable/Source/CustomizableObjectEditor/Private/MuCOE/Nodes/CustomizableObjectNodeObject.cpp b/Engine/Plugins/Mutable/Source/CustomizableObjectEditor/Private/MuCOE/Nodes/CustomizableObjectNodeObject.cpp
--- a/Engine/Plugins/Mutable/Source/CustomizableObjectEditor/Private/MuCOE/Nodes/CustomizableObjectNodeObject.cpp
+++ b/Engine/Plugins/Mutable/Source/CustomizableObjectEditor/Private/MuCOE/Nodes/CustomizableObjectNodeObject.cpp
@@ -3,6 +3,7 @@
 #include "MuCOE/Nodes/CustomizableObjectNodeObject.h"
 
 #include "CustomizableObjectNodeStaticString.h"
+#include "CONodeSwitch.h"
 #include "MuCO/ICustomizableObjectModule.h"
 #include "MuCO/CustomizableObjectExtension.h"
 #include "MuCO/CustomizableObjectPrivate.h"
@@ -19,7 +20,6 @@
 #include "MuCO/CustomizableObjectCustomVersion.h"
 #include "MuCOE/CustomizableObjectEditor_Deprecated.h"
 #include "MuCOE/Nodes/CustomizableObjectNodeComponentMeshAddTo.h"
-#include "MuCOE/Nodes/CONodeSkeletalMeshSectionSwitch.h"
 #include "Settings/EditorStyleSettings.h"
 
 #include UE_INLINE_GENERATED_CPP_BY_NAME(CustomizableObjectNodeObject)
@@ -346,7 +346,7 @@ void UCustomizableObjectNodeObject::BackwardsCompatibleFixup(int32 CustomizableO
 
 					NodeComponentMeshAddTo = CreateNodeComponent(NodeMaterial->MeshComponentName_DEPRECATED);
 				}
-				else if (UCONodeSkeletalMeshSectionSwitch* NodeMaterialSwitch = Cast<UCONodeSkeletalMeshSectionSwitch>(LinkedNode))
+				else if (const UCONodeSwitch* NodeMaterialSwitch = CastSwitch(LinkedNode, UEdGraphSchema_CustomizableObject::PC_MeshSection))
 				{
 					bFixNode = true;
 
diff --git a/Engine/Plugins/Mutable/Source/CustomizableObjectEditor/Private/MuCOE/Nodes/CustomizableObjectNodeParameter.cpp b/Engine/Plugins/Mutable/Source/CustomizableObjectEditor/Private/MuCOE/Nodes/CustomizableObjectNodeParameter.cpp
--- a/Engine/Plugins/Mutable/Source/CustomizableObjectEditor/Private/MuCOE/Nodes/CustomizableObjectNodeParameter.cpp
+++ b/Engine/Plugins/Mutable/Source/CustomizableObjectEditor/Private/MuCOE/Nodes/CustomizableObjectNodeParameter.cpp
@@ -88,7 +88,7 @@ void UCustomizableObjectNodeParameter::SetPinEditableName(const UEdGraphPin& Pin
 
 FText UCustomizableObjectNodeParameter::GetNodeTitle(ENodeTitleType::Type TitleType) const
 {
-	return FText::Format(LOCTEXT("ParameterTitle_ListView", "{0} Parameter"), UEdGraphSchema_CustomizableObject::GetPinCategoryFriendlyName(GetCategory()));
+	return FText::Format(LOCTEXT("ParameterTitle_ListView", "{0} Parameter"), GetCategoryFriendlyName());
 }
 
 
diff --git a/Engine/Plugins/Mutable/Source/CustomizableObjectEditor/Private/MuCOE/Nodes/CustomizableObjectNodePassThroughTextureSwitch.cpp b/Engine/Plugins/Mutable/Source/CustomizableObjectEditor/Private/MuCOE/Nodes/CustomizableObjectNodePassThroughTextureSwitch.cpp
--- a/Engine/Plugins/Mutable/Source/CustomizableObjectEditor/Private/MuCOE/Nodes/CustomizableObjectNodePassThroughTextureSwitch.cpp
+++ b/Engine/Plugins/Mutable/Source/CustomizableObjectEditor/Private/MuCOE/Nodes/CustomizableObjectNodePassThroughTextureSwitch.cpp
@@ -2,6 +2,7 @@
 
 #include "MuCOE/Nodes/CustomizableObjectNodePassThroughTextureSwitch.h"
 
+#include "MuCO/CustomizableObjectCustomVersion.h"
 #include "MuCOE/EdGraphSchema_CustomizableObject.h"
 
 #include UE_INLINE_GENERATED_CPP_BY_NAME(CustomizableObjectNodePassThroughTextureSwitch)
@@ -9,9 +10,11 @@
 #define LOCTEXT_NAMESPACE "CustomizableObjectEditor"
 
 
-FName UCustomizableObjectNodePassThroughTextureSwitch::GetCategory() const
+void UCustomizableObjectNodePassThroughTextureSwitch::BackwardsCompatibleFixup(int32 CustomizableObjectCustomVersion)
 {
-	return UEdGraphSchema_CustomizableObject::PC_Texture_Passthrough;
+	PinType = UEdGraphSchema_CustomizableObject::PC_Texture_Passthrough;
+
+	Super::BackwardsCompatibleFixup(CustomizableObjectCustomVersion);
 }
 
 
diff --git a/Engine/Plugins/Mutable/Source/CustomizableObjectEditor/Private/MuCOE/Nodes/CustomizableObjectNodePassThroughTextureSwitch.h b/Engine/Plugins/Mutable/Source/CustomizableObjectEditor/Private/MuCOE/Nodes/CustomizableObjectNodePassThroughTextureSwitch.h
--- a/Engine/Plugins/Mutable/Source/CustomizableObjectEditor/Private/MuCOE/Nodes/CustomizableObjectNodePassThroughTextureSwitch.h
+++ b/Engine/Plugins/Mutable/Source/CustomizableObjectEditor/Private/MuCOE/Nodes/CustomizableObjectNodePassThroughTextureSwitch.h
@@ -2,21 +2,21 @@
 
 #pragma once
 
-#include "MuCOE/Nodes/CustomizableObjectNodeSwitchBase.h"
+#include "MuCOE/Nodes/CONodeSwitch.h"
 
 #include "CustomizableObjectNodePassThroughTextureSwitch.generated.h"
 
 #define UE_API CUSTOMIZABLEOBJECTEDITOR_API
 
 
-UCLASS(MinimalAPI)
-class UCustomizableObjectNodePassThroughTextureSwitch : public UCustomizableObjectNodeSwitchBase
+UCLASS(NotPlaceable, HideDropdown)
+class UCustomizableObjectNodePassThroughTextureSwitch : public UCONodeSwitch
 {
 public:
 	GENERATED_BODY()
 	
-	// UCustomizableObjectNodeSwitchBase interface
-	UE_API virtual FName GetCategory() const override;
+	// UCustomizableObjectNodeSwitch interface
+	virtual void BackwardsCompatibleFixup(int32 CustomizableObjectCustomVersion) override;
 };
 
 #undef UE_API
diff --git a/Engine/Plugins/Mutable/Source/CustomizableObjectEditor/Private/MuCOE/Nodes/CustomizableObjectNodeSwitchBase.cpp b/Engine/Plugins/Mutable/Source/CustomizableObjectEditor/Private/MuCOE/Nodes/CustomizableObjectNodeSwitchBase.cpp
deleted file mode 100644
diff --git a/Engine/Plugins/Mutable/Source/CustomizableObjectEditor/Private/MuCOE/Nodes/CustomizableObjectNodeSwitchBase.h b/Engine/Plugins/Mutable/Source/CustomizableObjectEditor/Private/MuCOE/Nodes/CustomizableObjectNodeSwitchBase.h
deleted file mode 100644
diff --git a/Engine/Plugins/Mutable/Source/CustomizableObjectEditor/Private/MuCOE/Nodes/CustomizableObjectNodeTextureSwitch.cpp b/Engine/Plugins/Mutable/Source/CustomizableObjectEditor/Private/MuCOE/Nodes/CustomizableObjectNodeTextureSwitch.cpp
--- a/Engine/Plugins/Mutable/Source/CustomizableObjectEditor/Private/MuCOE/Nodes/CustomizableObjectNodeTextureSwitch.cpp
+++ b/Engine/Plugins/Mutable/Source/CustomizableObjectEditor/Private/MuCOE/Nodes/CustomizableObjectNodeTextureSwitch.cpp
@@ -12,6 +12,8 @@
 
 void UCustomizableObjectNodeTextureSwitch::BackwardsCompatibleFixup(int32 CustomizableObjectCustomVersion)
 {
+	PinType = UEdGraphSchema_CustomizableObject::PC_Texture;
+	
 	Super::BackwardsCompatibleFixup(CustomizableObjectCustomVersion);
 	
 	if (CustomizableObjectCustomVersion == FCustomizableObjectCustomVersion::FixPinsNamesImageToTexture2)
@@ -25,11 +27,5 @@ void UCustomizableObjectNodeTextureSwitch::BackwardsCompatibleFixup(int32 Custom
 }
 
 
-FName UCustomizableObjectNodeTextureSwitch::GetCategory() const
-{
-	return UEdGraphSchema_CustomizableObject::PC_Texture;
-}
-
-
 #undef LOCTEXT_NAMESPACE
 
diff --git a/Engine/Plugins/Mutable/Source/CustomizableObjectEditor/Private/MuCOE/Nodes/CustomizableObjectNodeTextureSwitch.h b/Engine/Plugins/Mutable/Source/CustomizableObjectEditor/Private/MuCOE/Nodes/CustomizableObjectNodeTextureSwitch.h
--- a/Engine/Plugins/Mutable/Source/CustomizableObjectEditor/Private/MuCOE/Nodes/CustomizableObjectNodeTextureSwitch.h
+++ b/Engine/Plugins/Mutable/Source/CustomizableObjectEditor/Private/MuCOE/Nodes/CustomizableObjectNodeTextureSwitch.h
@@ -2,23 +2,22 @@
 
 #pragma once
 
-#include "MuCOE/Nodes/CustomizableObjectNodeSwitchBase.h"
+#include "CONodeSwitch.h"
 
 #include "CustomizableObjectNodeTextureSwitch.generated.h"
 
 #define UE_API CUSTOMIZABLEOBJECTEDITOR_API
 
 
-UCLASS(MinimalAPI)
-class UCustomizableObjectNodeTextureSwitch : public UCustomizableObjectNodeSwitchBase
+// DEPRECATED
+UCLASS(NotPlaceable, HideDropdown)
+class UCustomizableObjectNodeTextureSwitch : public UCONodeSwitch
 {
 public:
 	GENERATED_BODY()
-	// UCustomizableObjectNode interface
-	UE_API virtual void BackwardsCompatibleFixup(int32 CustomizableObjectCustomVersion) override;
 	
-	// UCustomizableObjectNodeSwitchBase interface
-	UE_API virtual FName GetCategory() const override;
+	// UCustomizableObjectNode interface
+	virtual void BackwardsCompatibleFixup(int32 CustomizableObjectCustomVersion) override;
 };
 
 #undef UE_API
diff --git a/Engine/Plugins/Mutable/Source/CustomizableObjectEditor/Private/MuCOE/SMutableGraphViewer.cpp b/Engine/Plugins/Mutable/Source/CustomizableObjectEditor/Private/MuCOE/SMutableGraphViewer.cpp
--- a/Engine/Plugins/Mutable/Source/CustomizableObjectEditor/Private/MuCOE/SMutableGraphViewer.cpp
+++ b/Engine/Plugins/Mutable/Source/CustomizableObjectEditor/Private/MuCOE/SMutableGraphViewer.cpp
@@ -225,6 +225,7 @@ static const char* MutableNodeNames[] =
 	"External",
 	"ExternalOperation",
 	"ExternalParameter",
+	"ExternalSwitch",
 	
 	"ExtensionData",
 	"ExtensionDataConstant",
diff --git a/Engine/Plugins/Mutable/Source/CustomizableObjectEditor/Public/MuCOE/EdGraphSchema_CustomizableObject.h b/Engine/Plugins/Mutable/Source/CustomizableObjectEditor/Public/MuCOE/EdGraphSchema_CustomizableObject.h
--- a/Engine/Plugins/Mutable/Source/CustomizableObjectEditor/Public/MuCOE/EdGraphSchema_CustomizableObject.h
+++ b/Engine/Plugins/Mutable/Source/CustomizableObjectEditor/Public/MuCOE/EdGraphSchema_CustomizableObject.h
@@ -82,7 +82,6 @@ public:
 	static UE_API FName GetPinCategoryName(const UScriptStruct* Type);
 	static UE_API FName GetPinSubCategoryName(const FName& PinSubCategory);
 	static UE_API FText GetPinCategoryFriendlyName(const FName& PinCategory);
-	static UE_API FText GetPinCategoryFriendlyName(const UScriptStruct* Type);
 	static UE_API FText GetPinSubCategoryFriendlyName(const FName& PinSubcategory);
 	static UE_API bool IsPassthrough(const FName& PinCategory);
 	static UE_API bool IsTypeNative(const UScriptStruct* Type);
diff --git a/Engine/Plugins/Mutable/Source/MutableRuntime/Internal/MuR/Model.h b/Engine/Plugins/Mutable/Source/MutableRuntime/Internal/MuR/Model.h
--- a/Engine/Plugins/Mutable/Source/MutableRuntime/Internal/MuR/Model.h
+++ b/Engine/Plugins/Mutable/Source/MutableRuntime/Internal/MuR/Model.h
@@ -253,6 +253,9 @@ namespace UE::Mutable::Private
         /** Constant layout data */
 		TArray<TManagedPtr<const FLayout>> ConstantLayouts;
 
+    	/** Constant Instanced Struct */
+    	TArray<TManagedPtr<const FInstancedStruct>> ConstantInstancedStructs;
+    	
         /** Constant projectors */
 		TArray<FProjector> ConstantProjectors;
 
@@ -336,9 +339,7 @@ namespace UE::Mutable::Private
 		 */
     	UE_API int32 GetConstantImageRomId(int32 ConstantImageIndex, int32 LODIndex) const;
 
-#if WITH_EDITOR
-        UE_API FOperation::CONSTANT_EXTENSIONDATA_INDEX AddConstant(const TManagedPtr<const FExtensionData>& Data);
-    	
+#if WITH_EDITOR    	
         UE_API FOperation::ADDRESS AddConstant(const TManagedPtr<const FLayout>& Layout);
     	
         UE_API FOperation::ADDRESS AddConstant(const TManagedPtr<const FSkeleton>& Skeleton);
@@ -347,9 +348,9 @@ namespace UE::Mutable::Private
 
         UE_API uint32 AddConstant(const FMeshSocket& Socket);
 
-        UE_API FOperation::CONSTANT_LIST_UINT32_INDEX AddConstant(const TArray<FMeshSocket>& Sockets);
+        UE_API FOperation::CONSTANT_LIST_UINT32 AddConstant(const TArray<FMeshSocket>& Sockets);
     	
-        UE_API FOperation::CONSTANT_STRING_INDEX AddConstant(const FString& Str);
+        UE_API FOperation::CONSTANT_STRING AddConstant(const FString& Str);
     	
         UE_API FOperation::ADDRESS AddConstant(const TArray<uint32>& UInt32List);
     	
@@ -361,9 +362,9 @@ namespace UE::Mutable::Private
 
 		UE_API FOperation::ADDRESS AddConstant(const TArray<FName>& NameList);
 
-        UE_API FOperation::CONSTANT_MATRIX_INDEX AddConstant(const FMatrix44f& Constant);
+        UE_API FOperation::CONSTANT_MATRIX AddConstant(const FMatrix44f& Constant);
     	
-        UE_API FOperation::CONSTANT_SHAPE_INDEX AddConstant(const FShape& Constant);
+        UE_API FOperation::CONSTANT_SHAPE AddConstant(const FShape& Constant);
     	
         UE_API FOperation::ADDRESS AddConstant(const FProjector& Constant);
     	
diff --git a/Engine/Plugins/Mutable/Source/MutableRuntime/Internal/MuR/Operations.h b/Engine/Plugins/Mutable/Source/MutableRuntime/Internal/MuR/Operations.h
--- a/Engine/Plugins/Mutable/Source/MutableRuntime/Internal/MuR/Operations.h
+++ b/Engine/Plugins/Mutable/Source/MutableRuntime/Internal/MuR/Operations.h
@@ -545,6 +545,8 @@ namespace UE::Mutable::Private
 
         /** External operation */
 		IS_EXTERNAL,
+
+    	IS_SWITCH,
     
         IS_RANGE_END,
 
@@ -678,13 +680,12 @@ namespace UE::Mutable::Private
 
         typedef ADDRESS_TYPE ADDRESS;
         typedef ADDRESS_TYPE PARAMETER;
-        typedef ADDRESS_TYPE CONSTANT_LIST_PASSTHROUGH_ID;
-        typedef ADDRESS_TYPE CONSTANT_SHAPE_INDEX;
-    	typedef ADDRESS_TYPE CONSTANT_MATRIX_INDEX;
-    	typedef ADDRESS_TYPE CONSTANT_EXTENSIONDATA_INDEX;
-    	typedef ADDRESS_TYPE CONSTANT_NAME_ID;
-    	typedef ADDRESS_TYPE CONSTANT_STRING_INDEX;
-		typedef ADDRESS_TYPE CONSTANT_LIST_UINT32_INDEX;
+        typedef ADDRESS_TYPE CONSTANT_SHAPE;
+    	typedef ADDRESS_TYPE CONSTANT_MATRIX;
+    	typedef ADDRESS_TYPE CONSTANT_NAME;
+    	typedef ADDRESS_TYPE CONSTANT_STRING;
+		typedef ADDRESS_TYPE CONSTANT_LIST_UINT32;
+    	typedef ADDRESS_TYPE CONSTANT_LIST_PASSTHROUGH_ID;
 
         static constexpr uint32 AddressByteCodeOffsetBits = 26;
         static constexpr uint32 MaxAddressByteCodeOffset = (1 << AddressByteCodeOffsetBits) - 1;
@@ -712,7 +713,7 @@ namespace UE::Mutable::Private
 
     	struct MatrixConstantArgs
     	{
-    		CONSTANT_MATRIX_INDEX value;
+    		CONSTANT_MATRIX value;
     	};
 
         struct ResourceConstantArgs
@@ -753,13 +754,13 @@ namespace UE::Mutable::Private
 		{
 			PARAMETER Variable;
 
-			CONSTANT_LIST_UINT32_INDEX ImageParameterNames;
-			CONSTANT_LIST_UINT32_INDEX ImageParameterAddress;
+			CONSTANT_LIST_UINT32 ImageParameterNames;
+			CONSTANT_LIST_UINT32 ImageParameterAddress;
 		};
 
     	struct MaterialSkeletalMeshBreakArgs
     	{
-    		CONSTANT_STRING_INDEX SlotNameStringIndex;
+    		CONSTANT_STRING SlotNameStringIndex;
     		ADDRESS SkeletalMesh;
     	};
 
@@ -1245,7 +1246,7 @@ namespace UE::Mutable::Private
 		{
 			ADDRESS Source;
 			ADDRESS MorphShape;
-			CONSTANT_SHAPE_INDEX VertexSelectionShapeOrBone;
+			CONSTANT_SHAPE VertexSelectionShapeOrBone;
 
 			float Dist, Factor, MaxBoneRadius;
 
@@ -1284,7 +1285,7 @@ namespace UE::Mutable::Private
 
 		struct MeshMorphArgs
 		{
-			CONSTANT_STRING_INDEX Name;
+			CONSTANT_STRING Name;
 			ADDRESS Base;
 			ADDRESS Factor;
 		};
@@ -1362,14 +1363,14 @@ namespace UE::Mutable::Private
 
         	union NameOrUInt32ListIndex
         	{
-        		CONSTANT_NAME_ID NameIndex;
-        		CONSTANT_LIST_UINT32_INDEX ListIndex;
+        		CONSTANT_NAME NameIndex;
+        		CONSTANT_LIST_UINT32 ListIndex;
         	};
         	
             union StringAddressOrUInt32ListAddress
             {
                 ADDRESS StringAddress;
-				CONSTANT_LIST_UINT32_INDEX ListAddress;
+				CONSTANT_LIST_UINT32 ListAddress;
             };
 
             union ResourceIdOrUInt64ListAddress
@@ -1381,13 +1382,13 @@ namespace UE::Mutable::Private
 			union PassthroughIDOrUInt32ListAddress
 			{
 				PASSTHROUGH_ID PassthroughId;
-				CONSTANT_LIST_UINT32_INDEX ListAddress;
+				CONSTANT_LIST_UINT32 ListAddress;
 			};
 
 			union SocketIdOrUInt32ListAddress
 			{
 				uint32 SocketId;
-				CONSTANT_LIST_UINT32_INDEX ListAddress;
+				CONSTANT_LIST_UINT32 ListAddress;
 			};
 
             EnumFlags Flags; 
@@ -1445,7 +1446,7 @@ namespace UE::Mutable::Private
 		{
 			ADDRESS Instance;
 			ADDRESS Material;
-			CONSTANT_STRING_INDEX SlotName;
+			CONSTANT_STRING SlotName;
 		};
     	
     	struct InstanceAddOverlayMaterialArgs
@@ -1453,7 +1454,7 @@ namespace UE::Mutable::Private
     		ADDRESS Instance;
     		ADDRESS Material;
     		
-    		CONSTANT_STRING_INDEX SlotName;
+    		CONSTANT_STRING SlotName;
     	};
 
         struct LayoutPackArgs
@@ -1486,14 +1487,14 @@ namespace UE::Mutable::Private
 		struct MaterialBreakArgs
 		{
 			ADDRESS Material;
-			CONSTANT_STRING_INDEX ParameterName;
+			CONSTANT_STRING ParameterName;
 			int8 LayerIndex;
 		};
 
 		struct MaterialBreakImageParameterArgs
 		{
 			ADDRESS MaterialParameter;
-			CONSTANT_STRING_INDEX ParameterName;
+			CONSTANT_STRING ParameterName;
 			int8 LayerIndex;
 		};
 
@@ -1560,6 +1561,7 @@ namespace UE::Mutable::Private
         case EDataType::Scalar: return EOpType::SC_SWITCH;
         case EDataType::Int: return EOpType::NU_SWITCH;
 		case EDataType::ExtensionData: return EOpType::ED_SWITCH;
+		case EDataType::InstancedStruct: return EOpType::IS_SWITCH;
 
         default:
 			check(false);
diff --git a/Engine/Plugins/Mutable/Source/MutableRuntime/Private/MuR/CodeRunner.cpp b/Engine/Plugins/Mutable/Source/MutableRuntime/Private/MuR/CodeRunner.cpp
--- a/Engine/Plugins/Mutable/Source/MutableRuntime/Private/MuR/CodeRunner.cpp
+++ b/Engine/Plugins/Mutable/Source/MutableRuntime/Private/MuR/CodeRunner.cpp
@@ -385,18 +385,19 @@ namespace UE::Mutable::Private
 			{
 				switch (GetOpDataType(Type))
 				{
-				case EDataType::Bool:			StoreBool(Item, false); break;
-				case EDataType::Int:			StoreInt(Item, 0); break;
-				case EDataType::Scalar:			StoreScalar(Item, 0.0f); break;
-				case EDataType::String:			StoreString(Item, nullptr); break;
-				case EDataType::Color:			StoreColor(Item, FVector4f()); break;
-				case EDataType::Projector:		StoreProjector(Item, FProjector()); break;
-				case EDataType::Mesh:			StoreMesh(Item, nullptr); break;
-				case EDataType::Image:			StoreImage(Item, nullptr); break;
-				case EDataType::Layout:			StoreLayout(Item, nullptr); break;
-				case EDataType::Instance:		StoreInstance(Item, nullptr); break;
-				case EDataType::ExtensionData:	StoreExtensionData(Item, MakeManaged<FExtensionData>()); break;
-				case EDataType::Material:		StoreMaterial(Item, MakeManaged<FMaterial>()); break;
+				case EDataType::Bool:				StoreBool(Item, false); break;
+				case EDataType::Int:				StoreInt(Item, 0); break;
+				case EDataType::Scalar:				StoreScalar(Item, 0.0f); break;
+				case EDataType::String:				StoreString(Item, nullptr); break;
+				case EDataType::Color:				StoreColor(Item, FVector4f()); break;
+				case EDataType::Projector:			StoreProjector(Item, FProjector()); break;
+				case EDataType::Mesh:				StoreMesh(Item, nullptr); break;
+				case EDataType::Image:				StoreImage(Item, nullptr); break;
+				case EDataType::Layout:				StoreLayout(Item, nullptr); break;
+				case EDataType::Instance:			StoreInstance(Item, nullptr); break;
+				case EDataType::ExtensionData:		StoreExtensionData(Item, MakeManaged<FExtensionData>()); break;
+				case EDataType::Material:			StoreMaterial(Item, MakeManaged<FMaterial>()); break;
+				case EDataType::InstancedStruct:	StoreInstancedStruct(Item, MakeManaged<FInstancedStruct>()); break;
 				default:
 					// Not implemented
 					check(false);
@@ -472,18 +473,19 @@ namespace UE::Mutable::Private
             FCacheAddress ResultAddress(ResultAt, Item);
             switch (GetOpDataType(Type))
             {
-            case EDataType::Bool:			StoreBool(ItemAddress, LoadBool(ResultAddress)); break;
-            case EDataType::Int:			StoreInt(ItemAddress, LoadInt(ResultAddress)); break;
-            case EDataType::Scalar:			StoreScalar(ItemAddress, LoadScalar(ResultAddress)); break;
-            case EDataType::String:			StoreString(ItemAddress, LoadString(ResultAddress)); break;
-            case EDataType::Color:			StoreColor(ItemAddress, LoadColor(ResultAddress)); break;
-            case EDataType::Projector:		StoreProjector(ItemAddress, LoadProjector(ResultAddress)); break;
-			case EDataType::Mesh:			StoreMesh(ItemAddress, LoadMesh(ResultAddress)); break;
-            case EDataType::Image:			StoreImage(ItemAddress, LoadImage(ResultAddress)); break;
-            case EDataType::Layout:			StoreLayout(ItemAddress, LoadLayout(ResultAddress)); break;
-            case EDataType::Instance:		StoreInstance(ItemAddress, LoadInstance(ResultAddress)); break;
-			case EDataType::ExtensionData:	StoreExtensionData(ItemAddress, LoadExtensionData(ResultAddress)); break;
-			case EDataType::Material:		StoreMaterial(ItemAddress, LoadMaterial(ResultAddress)); break;
+            case EDataType::Bool:				StoreBool(ItemAddress, LoadBool(ResultAddress)); break;
+            case EDataType::Int:				StoreInt(ItemAddress, LoadInt(ResultAddress)); break;
+            case EDataType::Scalar:				StoreScalar(ItemAddress, LoadScalar(ResultAddress)); break;
+            case EDataType::String:				StoreString(ItemAddress, LoadString(ResultAddress)); break;
+            case EDataType::Color:				StoreColor(ItemAddress, LoadColor(ResultAddress)); break;
+            case EDataType::Projector:			StoreProjector(ItemAddress, LoadProjector(ResultAddress)); break;
+			case EDataType::Mesh:				StoreMesh(ItemAddress, LoadMesh(ResultAddress)); break;
+            case EDataType::Image:				StoreImage(ItemAddress, LoadImage(ResultAddress)); break;
+            case EDataType::Layout:				StoreLayout(ItemAddress, LoadLayout(ResultAddress)); break;
+            case EDataType::Instance:			StoreInstance(ItemAddress, LoadInstance(ResultAddress)); break;
+			case EDataType::ExtensionData:		StoreExtensionData(ItemAddress, LoadExtensionData(ResultAddress)); break;
+			case EDataType::Material:			StoreMaterial(ItemAddress, LoadMaterial(ResultAddress)); break;
+			case EDataType::InstancedStruct:	StoreInstancedStruct(ItemAddress, LoadInstancedStruct(ResultAddress)); break;
             default:
                 // Not implemented
                 check(false);
@@ -7120,6 +7122,7 @@ namespace UE::Mutable::Private
         case EOpType::IN_SWITCH:
 		case EOpType::ED_SWITCH:
 		case EOpType::MI_SWITCH:
+		case EOpType::IS_SWITCH:
             RunCode_Switch(Item);
             break;
 
@@ -7154,7 +7157,7 @@ namespace UE::Mutable::Private
 					
 				const FExternalOperation* Operation = OperationInstancedStruct.GetPtr<FExternalOperation>();
 
-				const TArray<TPair<FName, const UScriptStruct*>>& Inputs = Operation->GetInputs();
+				const TArray<TPair<FText, const UScriptStruct*>>& Inputs = Operation->GetInputs();
 				if (Inputs.Num() != Args.NumOperants)
 				{
 					LiveInstance->UpdateLogger->LogUpdateMessage(FString::Printf(TEXT("External Operation inputs do not match compiled data: %s"), *OperationInstancedStruct.GetScriptStruct()->GetName()));
@@ -7194,7 +7197,7 @@ namespace UE::Mutable::Private
 							FMemory::Memcpy(&OperantAt, Data, sizeof(FOperation::ADDRESS));
 							Data += sizeof(FOperation::ADDRESS);
 
-							const TPair<FName, const UScriptStruct*>& Input = Inputs[Index];
+							const TPair<FText, const UScriptStruct*>& Input = Inputs[Index];
 
 							const EOpType InputOpType = Program.GetOpType(OperantAt);
 							const EDataType InputDataType = GetOpDataType(InputOpType);
@@ -7208,7 +7211,7 @@ namespace UE::Mutable::Private
 									Value.Ptr = MakeManaged<FInstancedStruct>();
 									Value.Ptr->InitializeAs<FVectorAdapter>().Value = Vector;
 
-									Context.Inputs.Add(Input.Key, MoveTemp(Value));
+									Context.Inputs.Add(Input.Key.ToString(), MoveTemp(Value));
 									break;
 								}
 								
@@ -7220,7 +7223,7 @@ namespace UE::Mutable::Private
 									Value.Ptr = MakeManaged<FInstancedStruct>();
 									Value.Ptr->InitializeAs<FTextureAdapter>().Image = ConstCastManagedPtr<FImage>(Image);
 									
-									Context.Inputs.Add(Input.Key, MoveTemp(Value));
+									Context.Inputs.Add(Input.Key.ToString(), MoveTemp(Value));
 									break;
 								}
 
@@ -7232,7 +7235,7 @@ namespace UE::Mutable::Private
 									Value.Ptr = MakeManaged<FInstancedStruct>();
 									Value.Ptr->InitializeAs<FMeshAdapter>().Mesh = ConstCastManagedPtr<FMesh>(Mesh);
 
-									Context.Inputs.Add(Input.Key, MoveTemp(Value));
+									Context.Inputs.Add(Input.Key.ToString(), MoveTemp(Value));
 									break;
 								}
 								
@@ -7244,7 +7247,7 @@ namespace UE::Mutable::Private
 									Value.Ptr = MakeManaged<FInstancedStruct>();
 									Value.Ptr->InitializeAs<FMaterialAdapter>().Material = ConstCastManagedPtr<FMaterial>(Material);
 
-									Context.Inputs.Add(Input.Key, MoveTemp(Value));
+									Context.Inputs.Add(Input.Key.ToString(), MoveTemp(Value));
 									break;
 								}
 								
@@ -7256,7 +7259,7 @@ namespace UE::Mutable::Private
 									Value.Ptr = MakeManaged<FInstancedStruct>();
 									Value.Ptr = ConstCastManagedPtr<FInstancedStruct>(InstancedStruct);
 									
-									Context.Inputs.Add(Input.Key, MoveTemp(Value));
+									Context.Inputs.Add(Input.Key.ToString(), MoveTemp(Value));
 									break;
 								}
 									
@@ -7268,7 +7271,7 @@ namespace UE::Mutable::Private
 									Value.Ptr = MakeManaged<FInstancedStruct>();
 									Value.Ptr->InitializeAs<FFloatAdapter>().Value = Float;
 
-									Context.Inputs.Add(Input.Key, MoveTemp(Value));
+									Context.Inputs.Add(Input.Key.ToString(), MoveTemp(Value));
 									break;
 								}
 									
diff --git a/Engine/Plugins/Mutable/Source/MutableRuntime/Private/MuR/CodeRunnerBegin.cpp b/Engine/Plugins/Mutable/Source/MutableRuntime/Private/MuR/CodeRunnerBegin.cpp
--- a/Engine/Plugins/Mutable/Source/MutableRuntime/Private/MuR/CodeRunnerBegin.cpp
+++ b/Engine/Plugins/Mutable/Source/MutableRuntime/Private/MuR/CodeRunnerBegin.cpp
@@ -92,7 +92,7 @@ namespace UE::Mutable::Private
 		{
 			return;
 		}
-
+		
 		Items.Push(Item);
 	}
  
@@ -555,6 +555,7 @@ namespace UE::Mutable::Private
 			case EOpType::NU_SWITCH:
 			case EOpType::SC_SWITCH:
 			case EOpType::MI_SWITCH:
+			case EOpType::IS_SWITCH:
 				{
 					const uint8* Data = Program.GetOpArgsPointer(Item.At);
 
@@ -615,6 +616,7 @@ namespace UE::Mutable::Private
 									case EDataType::Mesh:			
 									case EDataType::Image:			
 									case EDataType::Layout:			
+									case EDataType::InstancedStruct:			
 										StoreNone(Item);
 										break;
 								
@@ -726,6 +728,7 @@ namespace UE::Mutable::Private
 								case EDataType::Mesh:
 								case EDataType::Image:
 								case EDataType::Layout:
+								case EDataType::InstancedStruct:
 									StoreNone(Item);
 									break;
 
@@ -2807,7 +2810,7 @@ namespace UE::Mutable::Private
 
 					const FExternalOperation* Operation = OperationInstancedStruct.GetPtr<FExternalOperation>();
 
-					const TArray<TPair<FName, const UScriptStruct*>>& Inputs = Operation->GetInputs();
+					const TArray<TPair<FText, const UScriptStruct*>>& Inputs = Operation->GetInputs();
 					if (Inputs.Num() != Args.NumOperants)
 					{
 						StoreNone(Item);
@@ -2818,7 +2821,7 @@ namespace UE::Mutable::Private
 					{
 					case 0:
 						{
-							PushOp(FScheduledOpInline(Item.At, 1));
+							PushOp(FScheduledOpInline(Item.At, 1, Item.ExecutionIndex, false));
 
 							for (int32 Index = 0; Index < Args.NumOperants; ++Index)
 							{
diff --git a/Engine/Plugins/Mutable/Source/MutableRuntime/Private/MuR/External/MeshAdapter.cpp b/Engine/Plugins/Mutable/Source/MutableRuntime/Private/MuR/External/MeshAdapter.cpp
--- a/Engine/Plugins/Mutable/Source/MutableRuntime/Private/MuR/External/MeshAdapter.cpp
+++ b/Engine/Plugins/Mutable/Source/MutableRuntime/Private/MuR/External/MeshAdapter.cpp
@@ -56,4 +56,16 @@ namespace UE::Mutable
 	{
 		Private::MeshRemoveVerticesWithCullSet(Mesh.Get(), Vertices, false);
 	}
+
+
+	Private::TManagedPtr<Private::FMesh> FMeshAdapter::GetPrivate()
+	{
+		return Mesh;
+	}
+
+
+	Private::TManagedPtr<const Private::FMesh> FMeshAdapter::GetPrivate() const
+	{
+		return Mesh;
+	}
 }
diff --git a/Engine/Plugins/Mutable/Source/MutableRuntime/Private/MuR/External/Operation.cpp b/Engine/Plugins/Mutable/Source/MutableRuntime/Private/MuR/External/Operation.cpp
--- a/Engine/Plugins/Mutable/Source/MutableRuntime/Private/MuR/External/Operation.cpp
+++ b/Engine/Plugins/Mutable/Source/MutableRuntime/Private/MuR/External/Operation.cpp
@@ -5,9 +5,9 @@
 
 namespace UE::Mutable
 {
-	FValueConst FContext::GetInput(const FName& Name)
+	FValueConst FContext::GetInput(const FText& Name)
 	{
-		return Inputs.FindAndRemoveChecked(Name);
+		return Inputs.FindAndRemoveChecked(Name.ToString());
 	}
 
 	void FContext::SetOutput(FValueConst&& Value)
diff --git a/Engine/Plugins/Mutable/Source/MutableRuntime/Private/MuR/Model.cpp b/Engine/Plugins/Mutable/Source/MutableRuntime/Private/MuR/Model.cpp
--- a/Engine/Plugins/Mutable/Source/MutableRuntime/Private/MuR/Model.cpp
+++ b/Engine/Plugins/Mutable/Source/MutableRuntime/Private/MuR/Model.cpp
@@ -327,7 +327,7 @@ namespace UE::Mutable::Private
 	}
 
 
-	FOperation::CONSTANT_LIST_UINT32_INDEX FProgram::AddConstant(const TArray<FMeshSocket>& Sockets)
+	FOperation::CONSTANT_LIST_UINT32 FProgram::AddConstant(const TArray<FMeshSocket>& Sockets)
 	{
 		TArray<uint32> SocketIds;
 		SocketIds.Reset(Sockets.Num());
@@ -341,9 +341,9 @@ namespace UE::Mutable::Private
 		return AddConstant(SocketIds);
 	}
 
-    FOperation::CONSTANT_STRING_INDEX FProgram::AddConstant(const FString& Str)
+    FOperation::CONSTANT_STRING FProgram::AddConstant(const FString& Str)
     {     
-	    return (FOperation::CONSTANT_STRING_INDEX)ConstantStrings.AddUnique(Str);
+	    return (FOperation::CONSTANT_STRING)ConstantStrings.AddUnique(Str);
     }
 
 
@@ -408,35 +408,35 @@ namespace UE::Mutable::Private
     }
 
 
-    FOperation::CONSTANT_MATRIX_INDEX FProgram::AddConstant(const FMatrix44f& Constant)
+    FOperation::CONSTANT_MATRIX FProgram::AddConstant(const FMatrix44f& Constant)
     {
 	    // Ensure unique
 	    for (int32 i = 0; i < ConstantMatrices.Num(); ++i)
 	    {
 		    if (ConstantMatrices[i] == Constant)
 		    {
-			    return (FOperation::CONSTANT_MATRIX_INDEX)i;
+			    return (FOperation::CONSTANT_MATRIX)i;
 		    }
 	    }
 
-	    FOperation::CONSTANT_MATRIX_INDEX Index = FOperation::CONSTANT_MATRIX_INDEX(ConstantMatrices.Num());
+	    FOperation::CONSTANT_MATRIX Index = FOperation::CONSTANT_MATRIX(ConstantMatrices.Num());
 	    ConstantMatrices.Add(Constant);
 	    return Index;
     }
 
 
-    FOperation::CONSTANT_SHAPE_INDEX FProgram::AddConstant(const FShape& Constant)
+    FOperation::CONSTANT_SHAPE FProgram::AddConstant(const FShape& Constant)
     {
 	    // Ensure unique
 	    for (int32 Index = 0; Index<ConstantShapes.Num(); ++Index)
 	    {
 		    if (ConstantShapes[Index] == Constant)
 		    {
-			    return (FOperation::CONSTANT_SHAPE_INDEX)Index;
+			    return (FOperation::CONSTANT_SHAPE)Index;
 		    }
 	    }
 
-	    FOperation::CONSTANT_SHAPE_INDEX Index = FOperation::CONSTANT_SHAPE_INDEX(ConstantShapes.Num());
+	    FOperation::CONSTANT_SHAPE Index = FOperation::CONSTANT_SHAPE(ConstantShapes.Num());
 	    ConstantShapes.Add(Constant);
 	    return Index;
     }
@@ -763,9 +763,9 @@ namespace UE::Mutable::Private
 
         return ResourceIndex.Index;
     }
+	
 
-
-	TSharedPtr<FParameters> FModel::NewParameters(TSharedPtr<const FModel> Model, const FParameters* OldParameters)
+    TSharedPtr<FParameters> FModel::NewParameters(TSharedPtr<const FModel> Model, const FParameters* OldParameters)
 	{
 		LLM_SCOPE_BYNAME(TEXT("MutableRuntime"));
 
diff --git a/Engine/Plugins/Mutable/Source/MutableRuntime/Private/MuR/Operations.cpp b/Engine/Plugins/Mutable/Source/MutableRuntime/Private/MuR/Operations.cpp
--- a/Engine/Plugins/Mutable/Source/MutableRuntime/Private/MuR/Operations.cpp
+++ b/Engine/Plugins/Mutable/Source/MutableRuntime/Private/MuR/Operations.cpp
@@ -79,6 +79,7 @@ namespace UE::Mutable::Private
         case EOpType::IN_SWITCH:
 		case EOpType::ED_SWITCH:
 		case EOpType::MI_SWITCH:
+		case EOpType::IS_SWITCH:
         {
 			const uint8* Data = program.GetOpArgsPointer(at);
 			
diff --git a/Engine/Plugins/Mutable/Source/MutableRuntime/Public/MuR/External/MeshAdapter.h b/Engine/Plugins/Mutable/Source/MutableRuntime/Public/MuR/External/MeshAdapter.h
--- a/Engine/Plugins/Mutable/Source/MutableRuntime/Public/MuR/External/MeshAdapter.h
+++ b/Engine/Plugins/Mutable/Source/MutableRuntime/Public/MuR/External/MeshAdapter.h
@@ -2,7 +2,7 @@
 
 #pragma once
 
-#include "MuR/Image.h"
+#include "MuR/ManagedPointer.h"
 
 #include "MeshAdapter.generated.h"
 
@@ -41,6 +41,10 @@ namespace UE::Mutable
 		
 		MUTABLERUNTIME_API void RemoveVertices(const TBitArray<>& Vertices);
 		
+		// Internal use only. Use at your own risk. Public API not maintained.
+		MUTABLERUNTIME_API Private::TManagedPtr<Private::FMesh> GetPrivate();
+		MUTABLERUNTIME_API Private::TManagedPtr<const Private::FMesh> GetPrivate() const;
+		
 	private:
 		Private::TManagedPtr<Private::FMesh> Mesh;
 	};
diff --git a/Engine/Plugins/Mutable/Source/MutableRuntime/Public/MuR/External/Operation.h b/Engine/Plugins/Mutable/Source/MutableRuntime/Public/MuR/External/Operation.h
--- a/Engine/Plugins/Mutable/Source/MutableRuntime/Public/MuR/External/Operation.h
+++ b/Engine/Plugins/Mutable/Source/MutableRuntime/Public/MuR/External/Operation.h
@@ -28,13 +28,13 @@ namespace UE::Mutable
 		/** Returns and removes the input from the Context.
 		 *
 		 * Once called the owner is the requester. This allows CopyOrMove to move (see CopyOrMove for details). */
-		MUTABLERUNTIME_API FValueConst GetInput(const FName& Name);
+		MUTABLERUNTIME_API FValueConst GetInput(const FText& Name);
 
 		/** Return the pointer to store the ouput. */
 		MUTABLERUNTIME_API void SetOutput(FValueConst&& Value);
 		
 	private:
-		TMap<const FName, FValueConst> Inputs;
+		TMap<const FString, FValueConst> Inputs;
 
 		FValueConst Output;
 	};
@@ -56,17 +56,14 @@ namespace UE::Mutable
 		
 		virtual ~FExternalOperation() = default;
 
-		/** Name of the operation. */
-		virtual FName GetName() const { return {}; }
-
 		/** Version of the operation. If it differs from the compiled one, it will automatically trigger a compilation. */
 		virtual uint32 GetVersion() const { return 0; }
 		
 		/** Static. It must always return the same inputs. */
-		virtual TArray<TPair<FName, const UScriptStruct*>> GetInputs() const { return {}; }
+		virtual TArray<TPair<FText, const UScriptStruct*>> GetInputs() const { return {}; }
 
 		/** Static. It must always return the same output. */
-		virtual TPair<FName, const UScriptStruct*> GetOutput() const { return {}; }
+		virtual TPair<FText, const UScriptStruct*> GetOutput() const { return {}; }
 
 		/** This function must be pure functional (same inputs, same output). */
 		virtual void Evaluate(FContext& Context) const {}
diff --git a/Engine/Plugins/Mutable/Source/MutableTools/Internal/MuT/Node.h b/Engine/Plugins/Mutable/Source/MutableTools/Internal/MuT/Node.h
--- a/Engine/Plugins/Mutable/Source/MutableTools/Internal/MuT/Node.h
+++ b/Engine/Plugins/Mutable/Source/MutableTools/Internal/MuT/Node.h
@@ -162,6 +162,7 @@ namespace UE::Mutable::Private
 			External,
 			ExternalOperation,
 			ExternalParameter,
+			ExternalSwitch,
 			
 			ExtensionData,
 			ExtensionDataConstant,
diff --git a/Engine/Plugins/Mutable/Source/MutableTools/Internal/MuT/NodeExternalSwitch.h b/Engine/Plugins/Mutable/Source/MutableTools/Internal/MuT/NodeExternalSwitch.h
new file mode 100644
--- /dev/null
+++ b/Engine/Plugins/Mutable/Source/MutableTools/Internal/MuT/NodeExternalSwitch.h
@@ -0,0 +1,36 @@
+// Copyright Epic Games, Inc. All Rights Reserved.
+
+#pragma once
+
+#include "MuT/NodeExternal.h"
+
+#include "MuR/Ptr.h"
+#include "MuR/RefCounted.h"
+#include "MuT/Node.h"
+#include "MuT/NodeScalar.h"
+
+#define UE_API MUTABLETOOLS_API
+
+
+namespace UE::Mutable::Private
+{
+    class NodeExternalSwitch : public NodeExternal
+	{
+	public:
+		Ptr<NodeScalar> Parameter;
+		TArray<Ptr<NodeExternal>> Options;
+
+		virtual const FNodeType* GetType() const override { return &StaticType; }
+		static const FNodeType* GetStaticType() { return &StaticType; }
+
+	protected:
+		/** Forbidden. Manage with the Ptr<> template. */
+		virtual ~NodeExternalSwitch() override = default;
+
+	private:
+		static UE_API FNodeType StaticType;
+	};
+}
+
+
+#undef UE_API
diff --git a/Engine/Plugins/Mutable/Source/MutableTools/Private/MuT/CodeGenerator.h b/Engine/Plugins/Mutable/Source/MutableTools/Private/MuT/CodeGenerator.h
--- a/Engine/Plugins/Mutable/Source/MutableTools/Private/MuT/CodeGenerator.h
+++ b/Engine/Plugins/Mutable/Source/MutableTools/Private/MuT/CodeGenerator.h
@@ -46,6 +46,7 @@
 
 namespace UE::Mutable::Private
 {
+	class NodeExternalSwitch;
 	extern TAutoConsoleVariable<bool> CVarAllowMulticompose;
 
 	class NodeImageParameterConvert;
@@ -942,6 +943,36 @@ namespace UE::Mutable::Private
     			
     			return Hash;
     		}
+    		
+    		FGenericGenerationOptions GetGenericOptions() const
+    		{
+    			if (MeshStaticOptions.IsSet())
+    			{
+    				return static_cast<FGenericGenerationOptions>(*MeshStaticOptions);
+    			}
+    			
+    			if (ImageOptions.IsSet())
+    			{
+    				return static_cast<FGenericGenerationOptions>(*ImageOptions);
+    			}
+    			
+    			if (ScalarOptions.IsSet())
+    			{
+    				return *ScalarOptions;
+    			}
+    			
+    			if (ColourOptions.IsSet())
+    			{
+    				return *ColourOptions;
+    			}
+    			
+    			if (MaterialOptions.IsSet())
+    			{
+    				return static_cast<FGenericGenerationOptions>(*MaterialOptions);
+    			}
+    			
+    			return {};
+    		}
     	};
     	
     	struct FGeneratedExternalCacheKey
@@ -978,7 +1009,9 @@ namespace UE::Mutable::Private
     	void GenerateExternal(FExternalGenerationResult& OutResult, const FExternalGenerationOptions& Options, const Ptr<const NodeExternal>& ExternalNode);
     	void GenerateExternal_Operation(FExternalGenerationResult& OutResult, const FExternalGenerationOptions& Options, const Ptr<const NodeExternalOperation>& ExternalNode);
     	void GenerateExternal_Parameter(FExternalGenerationResult& OutResult, const FExternalGenerationOptions& Options, const Ptr<const NodeExternalParameter>& ExtensionNode);
-
+    	void GenerateExternal_Switch(FExternalGenerationResult& OutResult, const FExternalGenerationOptions& Options, const Ptr<const NodeExternalSwitch>& ExtensionNode);
+    	Ptr<ASTOp> GenerateMissingExternalCode(const TCHAR* StrWhere, const void* ErrorContext);
+    	
 		//-----------------------------------------------------------------------------------------
     	
 		struct FExtensionDataGenerationResult
diff --git a/Engine/Plugins/Mutable/Source/MutableTools/Private/MuT/CodeGenerator_External.cpp b/Engine/Plugins/Mutable/Source/MutableTools/Private/MuT/CodeGenerator_External.cpp
--- a/Engine/Plugins/Mutable/Source/MutableTools/Private/MuT/CodeGenerator_External.cpp
+++ b/Engine/Plugins/Mutable/Source/MutableTools/Private/MuT/CodeGenerator_External.cpp
@@ -10,6 +10,7 @@
 #include "MuR/External/VectorAdapter.h"
 #include "MuT/NodeExternalOperation.h"
 #include "MuT/NodeExternalParameter.h"
+#include "MuT/NodeExternalSwitch.h"
 
 
 namespace UE::Mutable::Private
@@ -54,6 +55,13 @@ namespace UE::Mutable::Private
 				break;
 			}
 			
+		case Node::EType::ExternalSwitch:
+			{
+				Ptr<const NodeExternalSwitch> Parameter = static_pointer_cast<const NodeExternalSwitch>(ExternalNode);
+				GenerateExternal_Switch(OutResult, Options, Parameter);
+				break;
+			}
+			
 		default:
 			unimplemented();
 		}
@@ -186,28 +194,28 @@ namespace UE::Mutable::Private
 	}
 
 
-	void CodeGenerator::GenerateExternal_Parameter(FExternalGenerationResult& OutResult, const FExternalGenerationOptions& Options,	const Ptr<const NodeExternalParameter>& ExtensionNode)
+	void CodeGenerator::GenerateExternal_Parameter(FExternalGenerationResult& OutResult, const FExternalGenerationOptions& Options,	const Ptr<const NodeExternalParameter>& ExternalNode)
 	{
 		Ptr<ASTOpParameter> Op;
 
 		Ptr<ASTOpParameter>* Found = nullptr;
 		{
 			UE::TUniqueLock Lock(FirstPass.ParameterNodes.Mutex);
-			Found = FirstPass.ParameterNodes.GenericParametersCache.Find(ExtensionNode.get());
+			Found = FirstPass.ParameterNodes.GenericParametersCache.Find(ExternalNode.get());
 			if (!Found)
 			{
 				FParameterDesc Param;
-				Param.Name = ExtensionNode->Name;
-				bool bParseOk = FGuid::Parse(ExtensionNode->UID, Param.UID);
+				Param.Name = ExternalNode->Name;
+				bool bParseOk = FGuid::Parse(ExternalNode->UID, Param.UID);
 				check(bParseOk);
 				Param.Type = EParameterType::InstancedStruct;
-				Param.DefaultValue.Set<FParamInstancedStructType>(ExtensionNode->DefaultValue);
+				Param.DefaultValue.Set<FParamInstancedStructType>(ExternalNode->DefaultValue);
 
 				Op = new ASTOpParameter();
 				Op->Type = EOpType::IS_PARAMETER;
 				Op->Parameter = Param;
 
-				FirstPass.ParameterNodes.GenericParametersCache.Add(ExtensionNode.get(), Op);
+				FirstPass.ParameterNodes.GenericParametersCache.Add(ExternalNode.get(), Op);
 			}
 			else
 			{
@@ -218,7 +226,7 @@ namespace UE::Mutable::Private
 		if (!Found)
 		{
 			// Generate the code for the ranges
-			for (int32 RangeIndex = 0; RangeIndex < ExtensionNode->Ranges.Num(); ++RangeIndex)
+			for (int32 RangeIndex = 0; RangeIndex < ExternalNode->Ranges.Num(); ++RangeIndex)
 			{
 				FGenericGenerationOptions GenericOptions;
 
@@ -244,11 +252,72 @@ namespace UE::Mutable::Private
 				}
 				
 				FRangeGenerationResult rangeResult;
-				GenerateRange(rangeResult, GenericOptions, ExtensionNode->Ranges[RangeIndex]);
+				GenerateRange(rangeResult, GenericOptions, ExternalNode->Ranges[RangeIndex]);
 				Op->Ranges.Emplace(Op.get(), rangeResult.sizeOp, rangeResult.rangeName, rangeResult.rangeUID);
 			}
 		}
 
 		OutResult.Op = Op;
 	}
+
+
+	void CodeGenerator::GenerateExternal_Switch(FExternalGenerationResult& OutResult, const FExternalGenerationOptions& Options, const Ptr<const NodeExternalSwitch>& ExternalNode)
+	{
+		if (ExternalNode->Options.Num() == 0)
+		{
+			// No options in the switch!
+			Ptr<ASTOp> MissingOp = GenerateMissingExternalCode(TEXT("Switch option"),
+				ExternalNode->GetMessageContext());
+			
+			OutResult.Op = MissingOp;
+			return;
+		}
+
+		Ptr<ASTOpSwitch> Op = new ASTOpSwitch();
+		Op->Type = EOpType::IS_SWITCH;
+
+		// Variable value
+		if (ExternalNode->Parameter)
+		{
+			FScalarGenerationResult ParamResult;
+			GenerateScalar(ParamResult, Options.GetGenericOptions(), ExternalNode->Parameter.get());
+			Op->Variable = ParamResult.op;
+		}
+		else
+		{
+			// This argument is required
+			Op->Variable = GenerateMissingScalarCode(TEXT("Switch variable"), 0.0f, ExternalNode->GetMessageContext());
+		}
+
+		// Options
+		for (int32 OptionIndex = 0; OptionIndex < ExternalNode->Options.Num(); ++OptionIndex)
+		{
+			Ptr<ASTOp> Branch;
+			if (ExternalNode->Options[OptionIndex])
+			{
+				FExternalGenerationResult ChildResult;
+				GenerateExternal(ChildResult, Options, ExternalNode->Options[OptionIndex].get());
+				Branch = ChildResult.Op;
+			}
+			else
+			{
+				// This argument is required
+				Branch = GenerateMissingExternalCode(TEXT("Switch option"), ExternalNode->GetMessageContext());
+			}
+			
+			Op->Cases.Emplace(static_cast<int16_t>(OptionIndex), Op, Branch);
+		}
+
+		OutResult.Op = Op;
+	}
+	
+	
+	Ptr<ASTOp> CodeGenerator::GenerateMissingExternalCode(const TCHAR* StrWhere, const void* ErrorContext)
+	{
+		// Log a warning
+		const FString Msg = FString::Printf(TEXT("Required connection not found: %s"), StrWhere);
+		ErrorLog->Add(Msg, ELMT_ERROR, ErrorContext);
+		
+		return {};
+	}
 }
diff --git a/Engine/Plugins/Mutable/Source/MutableTools/Private/MuT/Compiler.cpp b/Engine/Plugins/Mutable/Source/MutableTools/Private/MuT/Compiler.cpp
--- a/Engine/Plugins/Mutable/Source/MutableTools/Private/MuT/Compiler.cpp
+++ b/Engine/Plugins/Mutable/Source/MutableTools/Private/MuT/Compiler.cpp
@@ -202,6 +202,7 @@ namespace UE::Mutable::Private
 
 		{ false,	{ 0, 0, 0, 0, 0, 0, 0, 0, 0, 0, 0, 0, 0, 0, 0, 0, 0, 0, 0, 0, 0, 0, 0, 0, 0, 0, 0, 0, 0 }		},	// IS_PARAMETER
     	{ true,		{ 0, 0, 0, 0, 0, 0, 0, 0, 0, 0, 0, 0, 0, 0, 0, 0, 0, 0, 0, 0, 0, 0, 0, 0, 0, 0, 0, 0, 0 }		},	// IS_EXTERNAL
+		{ false,	{ 0, 0, 0, 0, 0, 0, 0, 0, 0, 0, 0, 0, 0, 0, 0, 0, 0, 0, 0, 0, 0, 0, 0, 0, 0, 0, 0, 0, 0 }		},	// IS_SWITCH
 	};
 
 	// clang-format on
diff --git a/Engine/Plugins/Mutable/Source/MutableTools/Private/MuT/ErrorLog.cpp b/Engine/Plugins/Mutable/Source/MutableTools/Private/MuT/ErrorLog.cpp
--- a/Engine/Plugins/Mutable/Source/MutableTools/Private/MuT/ErrorLog.cpp
+++ b/Engine/Plugins/Mutable/Source/MutableTools/Private/MuT/ErrorLog.cpp
@@ -339,6 +339,7 @@ namespace UE::Mutable::Private
 
 		TEXT("IS_PARAMETER"),
 		TEXT("IS_EXTERNAL"),
+		TEXT("IS_SWITCH"),
 	};
 
 	static_assert(UE_ARRAY_COUNT(s_opNames) == int32(EOpType::COUNT));
diff --git a/Engine/Plugins/Mutable/Source/MutableTools/Private/MuT/Node.cpp b/Engine/Plugins/Mutable/Source/MutableTools/Private/MuT/Node.cpp
--- a/Engine/Plugins/Mutable/Source/MutableTools/Private/MuT/Node.cpp
+++ b/Engine/Plugins/Mutable/Source/MutableTools/Private/MuT/Node.cpp
@@ -96,6 +96,7 @@
 #include "MuT/NodeExternal.h"
 #include "MuT/NodeExternalOperation.h"
 #include "MuT/NodeExternalParameter.h"
+#include "MuT/NodeExternalSwitch.h"
 #include "MuT/NodeImageExternal.h"
 #include "MuT/NodeProjector.h"
 #include "MuT/NodeRange.h"
@@ -162,6 +163,7 @@ namespace UE::Mutable::Private
 	FNodeType NodeExternal::StaticType = FNodeType(Node::EType::External, Node::GetStaticType());
 	FNodeType NodeExternalOperation::StaticType = FNodeType(Node::EType::ExternalOperation, NodeExternal::GetStaticType());
 	FNodeType NodeExternalParameter::StaticType = FNodeType(Node::EType::ExternalParameter, NodeExternal::GetStaticType());
+	FNodeType NodeExternalSwitch::StaticType = FNodeType(Node::EType::ExternalSwitch, NodeExternal::GetStaticType());
 	
 	FNodeType NodeLOD::StaticType = FNodeType(Node::EType::LOD, Node::GetStaticType());
 
