[Mutable] Remove Mutable Game Thread tasks lock.

#Synced-CL 51468966

diff --git a/Engine/Plugins/Mutable/Source/CustomizableObject/Internal/MuCO/FMutableTaskGraph.h b/Engine/Plugins/Mutable/Source/CustomizableObject/Internal/MuCO/FMutableTaskGraph.h
--- a/Engine/Plugins/Mutable/Source/CustomizableObject/Internal/MuCO/FMutableTaskGraph.h
+++ b/Engine/Plugins/Mutable/Source/CustomizableObject/Internal/MuCO/FMutableTaskGraph.h
@@ -47,18 +47,15 @@ class FMutableTaskGraph
 		FString DebugName;
 		TUniqueFunction<void(float)> Body;
 		TArray<UE::Tasks::FTask> Prerequisites;
-		bool bLockMutableThread = false;
-	
+		
 		/** Check if all the dependencies of this task have been completed. */
 		bool AreDependenciesComplete() const;
 	};
 	
 public:
-	FMutableTaskGraph();
-
 	~FMutableTaskGraph();
 
-	void AddGameThreadTask(const TCHAR* DebugName, TUniqueFunction<void(float)>&& TaskBody, bool bLockMutableThread = false, const TArray<UE::Tasks::FTask>& Prerequisites = {});
+	void AddGameThreadTask(const TCHAR* DebugName, TUniqueFunction<void(float)>&& TaskBody, const TArray<UE::Tasks::FTask>& Prerequisites = {});
 	
 	/** Create and launch a task on the Mutable Thread with Normal priority. */
 	UE::Tasks::FTask AddMutableThreadTask(const TCHAR* DebugName, TUniqueFunction<void()>&& TaskBody, const TArray<UE::Tasks::FTask>& Prerequisites = {});
@@ -77,9 +74,6 @@ public:
 	  * Return false if not found or running. */
 	bool CancelMutableThreadTaskLowPriority(uint32 Id);
 	
-	/** Create and launch a task on Any Thread. */
-	void AddAnyThreadTask(const TCHAR* DebugName, TUniqueFunction<void()>&& TaskBody) const;
-
 	/** Wait for all Mutable Thread tasks. */
 	void WaitForMutableTasks();
 
@@ -95,21 +89,6 @@ public:
 	
 	int32 GetRemainingWork() const;
 
-private:
-	/** Lock, as soon as possible, the Mutable Thread. */
-	void AsyncLockMutableThread();
-
-	/** Check if the Mutable Thread is locked. */
-	bool IsMutableThreadLocked() const;
-
-	/** Check if the Mutable Thread is locked without using the lock. Does not update the region. */
-	bool IsMutableThreadLockedNoLock() const;
-
-public: 
-	/** Unlock the Mutable Thread. Releases all blocked tasks. */
-	void UnlockMutableThread(); // TODO Make private once UE-281921
-
-private:
 	/** A Mutable Task Low Priority will only be launched if:
 	 * - No other low priority task is running.
 	 * - Is allowed to launch low priority tasks.
@@ -148,14 +127,6 @@ private:
 
 	/** Last Mutable Task launched to the TaskGraph system. Low and normal priority. */
 	UE::Tasks::FTask LastMutableTask;
-
-	/** Last Mutable Task launched before locking the Mutable Thread. Once completed the Mutable Thread will be considered locked. */
-	UE::Tasks::FTask LastMutableTaskBeforeLock;
-	
-	/** Event to block upcoming tasks when the Mutable Thread is locked.
-	 *  Completed = Mutable Thread lock has been requested. May not be locked due to tasks still running (see LastMutableTaskBeforeLock).
-	 *  Not Completed = Mutable Thread unlocked. */
-	UE::Tasks::FTaskEvent MutableThreadUnlockEvent { UE_SOURCE_LOCATION };
 	
 	mutable FCriticalSection CriticalSectionLastGenerationTask;
 	mutable FCriticalSection CriticalSectionLastImageStreaming;
diff --git a/Engine/Plugins/Mutable/Source/CustomizableObject/Internal/MuCO/UnrealMutableModelDiskStreamer.h b/Engine/Plugins/Mutable/Source/CustomizableObject/Internal/MuCO/UnrealMutableModelDiskStreamer.h
--- a/Engine/Plugins/Mutable/Source/CustomizableObject/Internal/MuCO/UnrealMutableModelDiskStreamer.h
+++ b/Engine/Plugins/Mutable/Source/CustomizableObject/Internal/MuCO/UnrealMutableModelDiskStreamer.h
@@ -77,14 +77,16 @@ public:
 	CUSTOMIZABLEOBJECT_API bool IsReadCompleted(FOperationID) override;
 	CUSTOMIZABLEOBJECT_API bool EndRead(FOperationID) override;
 
-protected:
-
+private:
 	struct FReadRequest
 	{
 		TSharedPtr<IBulkDataIORequest> BulkReadRequest;
+		TSharedPtr<IAsyncReadFileHandle> FileHandle;
 		TSharedPtr<IAsyncReadRequest> FileReadRequest;
 		TSharedPtr<FAsyncFileCallBack> FileCallback;
 
+		uint32 FileId = 0;
+
 #if WITH_EDITORONLY_DATA
 		TSharedPtr<FDDCReadRequest> DDCReadRequest;
 #endif
@@ -97,7 +99,6 @@ protected:
 
 		FString BulkFilePrefix;
 		TMap<FOperationID, FReadRequest> CurrentReadRequests;
-		TMap<uint32, TSharedPtr<IAsyncReadFileHandle>> ReadFileHandles;
 
 		TWeakPtr<FModelStreamableBulkData> ModelStreamableBulkData;
 	};
diff --git a/Engine/Plugins/Mutable/Source/CustomizableObject/Private/MuCO/CustomizableObjectSystem.cpp b/Engine/Plugins/Mutable/Source/CustomizableObject/Private/MuCO/CustomizableObjectSystem.cpp
--- a/Engine/Plugins/Mutable/Source/CustomizableObject/Private/MuCO/CustomizableObjectSystem.cpp
+++ b/Engine/Plugins/Mutable/Source/CustomizableObject/Private/MuCO/CustomizableObjectSystem.cpp
@@ -2421,8 +2421,7 @@ namespace impl
 			[OperationData](float RemainingTime)
 			{
 				Task_Game_Callbacks(OperationData);
-			},
-			false);
+			});
 	}
 	
 	
@@ -2819,7 +2818,6 @@ namespace impl
 		{
 			Task_Game_ConvertResources(Operation);
 		},
-		false,
 		{Game_LoadUnrealAssets, Mutable_PrepareTextures});
 	}
 } // namespace impl
@@ -3013,8 +3011,7 @@ void UCustomizableObjectSystemPrivate::MainTick(const bool bBlocking)
 	if (IsEngineExitRequested())
 	{	
 		// TODO BEGIN. Remove once UE-281921  
-		MutableTaskGraph.UnlockMutableThread();
-   		MutableTaskGraph.AllowLaunchingMutableTaskLowPriority(true, false);
+		MutableTaskGraph.AllowLaunchingMutableTaskLowPriority(true, false);
 		// TODO END
 
 		return;
@@ -3548,7 +3545,6 @@ void UCustomizableObjectSystemPrivate::StartUpdate(const TSharedRef<FUpdateConte
 		{
 			impl::Task_Game_StartUpdate(Context);
 		},
-		true,
 		{ LastUpdateMutableTask });
 }
 
diff --git a/Engine/Plugins/Mutable/Source/CustomizableObject/Private/MuCO/FMutableTaskGraph.cpp b/Engine/Plugins/Mutable/Source/CustomizableObject/Private/MuCO/FMutableTaskGraph.cpp
--- a/Engine/Plugins/Mutable/Source/CustomizableObject/Private/MuCO/FMutableTaskGraph.cpp
+++ b/Engine/Plugins/Mutable/Source/CustomizableObject/Private/MuCO/FMutableTaskGraph.cpp
@@ -38,19 +38,6 @@ TAutoConsoleVariable<bool> CVarEnableParallelUpdates(
 	ECVF_Scalability);
 
 
-/** Status of the Mutable Thread lock. Only used for Insights Regions. */
-enum class ERegionLockStatus
-{
-	Requested,
-	Locked,
-	Unlocked // No region.
-};
-
-
-/** What Insights Region we are currently showing. Used only for debug. */
-static ERegionLockStatus RegionLockStatus;
-
-
 bool FMutableTaskGraph::FGameThreadTask::AreDependenciesComplete() const
 {
 	for (const UE::Tasks::FTask& Dependency : Prerequisites)
@@ -65,21 +52,15 @@ bool FMutableTaskGraph::FGameThreadTask::AreDependenciesComplete() const
 }
 
 
-FMutableTaskGraph::FMutableTaskGraph()
-{
-	MutableThreadUnlockEvent.Trigger();
-}
-
-
 FMutableTaskGraph::~FMutableTaskGraph()
 {
 	GameThreadTasks.Empty();
 }
 
 
-void FMutableTaskGraph::AddGameThreadTask(const TCHAR* DebugName, TUniqueFunction<void(float)>&& TaskBody, bool bLockMutableThread, const TArray<UE::Tasks::FTask>& Prerequisites)
+void FMutableTaskGraph::AddGameThreadTask(const TCHAR* DebugName, TUniqueFunction<void(float)>&& TaskBody, const TArray<UE::Tasks::FTask>& Prerequisites)
 {
-	GameThreadTasks.Enqueue({ DebugName, Forward<TUniqueFunction<void(float)>>(TaskBody), Prerequisites, bLockMutableThread });
+	GameThreadTasks.Enqueue({ DebugName, Forward<TUniqueFunction<void(float)>>(TaskBody), Prerequisites });
 }
 
 
@@ -93,9 +74,9 @@ UE::Tasks::FTask FMutableTaskGraph::AddMutableThreadTask(const TCHAR* DebugName,
 	{
 		FScopeLock Lock(&MutableTaskLock);
 
-		TArray<UE::Tasks::FTask, TInlineAllocator<2>> Prerequisites = LastMutableTask.IsValid() ?
-			TArray<UE::Tasks::FTask, TInlineAllocator<2>>{ LastMutableTask, MutableThreadUnlockEvent } :
-			TArray<UE::Tasks::FTask, TInlineAllocator<2>>{};
+		TArray<UE::Tasks::FTask, TInlineAllocator<1>> Prerequisites = LastMutableTask.IsValid() ?
+			TArray<UE::Tasks::FTask, TInlineAllocator<1>>{ LastMutableTask } :
+			TArray<UE::Tasks::FTask, TInlineAllocator<1>>{};
 
 		Prerequisites.Append(InPrerequisites);
 	
@@ -146,7 +127,6 @@ UE::Tasks::FTask FMutableTaskGraph::AddMeshStreamingTask(const TCHAR* DebugName,
 
 	Prerequisites.Append(InPrerequisites);
 	
-	
 	LastMeshStreamingTask = UE::Tasks::Launch(DebugName, Forward<TUniqueFunction<void()>>(TaskBody), Prerequisites, TASKGRAPH_PRIORITY);
 	return LastMeshStreamingTask;
 }
@@ -191,12 +171,6 @@ bool FMutableTaskGraph::CancelMutableThreadTaskLowPriority(uint32 Id)
 }
 
 
-void FMutableTaskGraph::AddAnyThreadTask(const TCHAR* DebugName, TUniqueFunction<void()>&& TaskBody) const
-{
-	UE::Tasks::Launch(DebugName, MoveTemp(TaskBody));
-}
-	
-
 void FMutableTaskGraph::WaitForMutableTasks()
 {
 	FScopeLock Lock(&MutableTaskLock);
@@ -365,17 +339,7 @@ void FMutableTaskGraph::TryLaunchGameThreadTask()
 		{
 			break;
 		}
-
-		if (PendingTask->bLockMutableThread)
-		{
-			AsyncLockMutableThread();
-
-			if (!IsMutableThreadLocked())
-			{
-				break;
-			}
-		}
-	
+		
 		FGameThreadTask CurrentTask;
 		GameThreadTasks.Dequeue(CurrentTask); // Dequeue now since current task can enqueue new tasks.
 
@@ -386,8 +350,6 @@ void FMutableTaskGraph::TryLaunchGameThreadTask()
 		RemainingTime = TickTime - (FPlatformTime::Seconds() - StartTime);
 		bLoop = RemainingTime > 0.0;
 	}
-
-	UnlockMutableThread();
 }
 
 
@@ -406,100 +368,3 @@ int32 FMutableTaskGraph::GetRemainingWork() const
 	return !IsTaskCompleted(LastMutableTask) + !IsTaskCompleted(LastMutableTaskLowPriority) + QueueMutableTasksLowPriority.Num() + !GameThreadTasks.IsEmpty();
 }
 
-
-void UpdateRegions(ERegionLockStatus NewRegionLockStatus)
-{
-	if (RegionLockStatus == NewRegionLockStatus)
-	{
-		return;
-	}
-
-	// End current region.
-	switch (RegionLockStatus)
-	{
-	case ERegionLockStatus::Requested:
-		TRACE_END_REGION(UE_MUTABLE_THREAD_REQUEST_LOCK);
-		break;
-	case ERegionLockStatus::Locked:
-		TRACE_END_REGION(UE_MUTABLE_THREAD_LOCK);
-		break;
-	case ERegionLockStatus::Unlocked:
-		break;
-	default:
-		unimplemented();
-	}
-
-	// Start new region.
-	switch (NewRegionLockStatus)
-	{
-	case ERegionLockStatus::Requested:
-		TRACE_BEGIN_REGION(UE_MUTABLE_THREAD_REQUEST_LOCK);
-		break;
-	case ERegionLockStatus::Locked:
-		TRACE_BEGIN_REGION(UE_MUTABLE_THREAD_LOCK);
-		break;
-	case ERegionLockStatus::Unlocked:
-		break;
-	default:
-		unimplemented();
-	}
-
-	RegionLockStatus = NewRegionLockStatus;
-}
-
-
-void FMutableTaskGraph::AsyncLockMutableThread()
-{
-	MUTABLE_CPUPROFILER_SCOPE(FMutableTaskGraph::AsyncLockMutableThread)
-
-	FScopeLock Lock(&MutableTaskLock);
-	
-	if (!MutableThreadUnlockEvent.IsCompleted()) // Locking already in progress.
-	{
-		return;
-	}
-
-	UpdateRegions(ERegionLockStatus::Requested);
-
-	LastMutableTaskBeforeLock = LastMutableTask;
-	MutableThreadUnlockEvent = UE::Tasks::FTaskEvent(UE_SOURCE_LOCATION);
-}
-
-
-bool FMutableTaskGraph::IsMutableThreadLockedNoLock() const
-{
-	return !MutableThreadUnlockEvent.IsCompleted() && LastMutableTaskBeforeLock.IsCompleted();
-}
-
-
-bool FMutableTaskGraph::IsMutableThreadLocked() const
-{
-	MUTABLE_CPUPROFILER_SCOPE(FMutableTaskGraph::IsMutableThreadLocked)
-
-	FScopeLock Lock(&MutableTaskLock);
-
-	const bool bLocked = IsMutableThreadLockedNoLock();
-
-	UpdateRegions(bLocked ? ERegionLockStatus::Locked : RegionLockStatus);
-	
-	return bLocked;
-}
-
-
-void FMutableTaskGraph::UnlockMutableThread()
-{
-	MUTABLE_CPUPROFILER_SCOPE(FMutableTaskGraph::UnlockMutableThread)
-	
-	FScopeLock Lock(&MutableTaskLock);
-
-	if (!IsMutableThreadLockedNoLock())
-	{
-		return;
-	}
-	
-	UpdateRegions(ERegionLockStatus::Unlocked);
-	
-	LastMutableTaskBeforeLock = {};
-	MutableThreadUnlockEvent.Trigger();
-}
-
diff --git a/Engine/Plugins/Mutable/Source/CustomizableObject/Private/MuCO/UnrealMutableModelDiskStreamer.cpp b/Engine/Plugins/Mutable/Source/CustomizableObject/Private/MuCO/UnrealMutableModelDiskStreamer.cpp
--- a/Engine/Plugins/Mutable/Source/CustomizableObject/Private/MuCO/UnrealMutableModelDiskStreamer.cpp
+++ b/Engine/Plugins/Mutable/Source/CustomizableObject/Private/MuCO/UnrealMutableModelDiskStreamer.cpp
@@ -70,9 +70,6 @@ bool FUnrealMutableModelBulkReader::PrepareStreamingForObject(UCustomizableObjec
 	// See if we can free previuously allocated resources.
 	for (int32 ObjectIndex = 0; ObjectIndex < Objects.Num(); )
 	{
-		// Close open file handles
-		Objects[ObjectIndex].ReadFileHandles.Empty();
-
 		if (!Objects[ObjectIndex].Model.Pin() && Objects[ObjectIndex].CurrentReadRequests.IsEmpty())
 		{
 			Objects.RemoveAtSwap(ObjectIndex);
@@ -427,22 +424,31 @@ UE::Mutable::Private::FModelReader::FOperationID FUnrealMutableModelBulkReader::
 				});
 		}
 
-		TSharedPtr<IAsyncReadFileHandle>& Found = ObjectData->ReadFileHandles.FindOrAdd(Block->FileId);
-		if (!Found)
+		TSharedPtr<IAsyncReadFileHandle> FileHandle;
+		for (const TTuple<FOperationID, FReadRequest>& Pair : ObjectData->CurrentReadRequests)
 		{
-			FString FilePath = Private::MakeBlockFilePath(ObjectData->BulkFilePrefix, *Block);
-			Found = MakeShareable(FPlatformFileManager::Get().GetPlatformFile().OpenAsyncRead(*FilePath));
-
-			if (!Found)
+			if (Pair.Value.FileId == Block->FileId)
+			{
+				FileHandle = Pair.Value.FileHandle;
+				break;
+			}
+		}
+		
+		if (!FileHandle)
+		{
+			const FString FilePath = Private::MakeBlockFilePath(ObjectData->BulkFilePrefix, *Block);
+			FileHandle = MakeShareable(FPlatformFileManager::Get().GetPlatformFile().OpenAsyncRead(*FilePath));
+			
+			if (!FileHandle)
 			{
 				UE_LOG(LogMutable, Error, TEXT("Failed to create AsyncReadFileHandle. File Path [%s]."), *FilePath);
 				check(false);
 				return -1;
 			}
 		}
-
-		TSharedPtr<IAsyncReadFileHandle> FileHandle = Found;
-
+	
+		ReadRequest.FileId = Block->FileId;
+		ReadRequest.FileHandle = FileHandle;
 		ReadRequest.FileReadRequest = MakeShareable(FileHandle->ReadRequest(
 			sizeof(MutableCompiledDataStreamHeader) + Block->Offset,
 			Size,
diff --git a/Engine/Plugins/Mutable/Source/MutableRuntime/Internal/MuR/RomManager.cpp b/Engine/Plugins/Mutable/Source/MutableRuntime/Internal/MuR/RomManager.cpp
--- a/Engine/Plugins/Mutable/Source/MutableRuntime/Internal/MuR/RomManager.cpp
+++ b/Engine/Plugins/Mutable/Source/MutableRuntime/Internal/MuR/RomManager.cpp
@@ -298,7 +298,7 @@ namespace UE::Mutable::Private
 		{
 			DataType = EDataType::Mesh;
 		}
-			
+		
 		RomLoadOp->StreamID = ModelReader->BeginReadBlock(&Model, RomIndex, RomLoadOp->StreamBuffer.GetData(), RomSize, DataType, &Callback);
 		if (RomLoadOp->StreamID < 0)
 		{
@@ -307,7 +307,7 @@ namespace UE::Mutable::Private
 		
 		Tasks::TTask<TManagedPtr<const Type>> Task = Tasks::Launch(UE_SOURCE_LOCATION, [Reader = ModelReader, RomLoadOp]() -> TManagedPtr<const Type>
 		{
-			bool bSuccess = Reader->EndRead(RomLoadOp->StreamID);
+			const bool bSuccess = Reader->EndRead(RomLoadOp->StreamID);
 			if (!bSuccess)
 			{
 				return nullptr;
