[Mutable] Fix extensions crash when closing the editor.

- Disabled Mutable Asset User Data plugin by default

#Synced-CL 51468966

diff --git a/Engine/Plugins/Experimental/MutableAssetUserData/MutableAssetUserData.uplugin b/Engine/Plugins/Experimental/MutableAssetUserData/MutableAssetUserData.uplugin
--- a/Engine/Plugins/Experimental/MutableAssetUserData/MutableAssetUserData.uplugin
+++ b/Engine/Plugins/Experimental/MutableAssetUserData/MutableAssetUserData.uplugin
@@ -14,7 +14,7 @@
 	"IsBetaVersion": false,
 	"IsExperimentalVersion": true,
 	"Installed": false,
-	"EnabledByDefault": true,
+	"EnabledByDefault": false,
 	"Modules": 
 	[
 		{
diff --git a/Engine/Plugins/Mutable/Source/CustomizableObject/Private/MuCO/CustomizableObjectModule.cpp b/Engine/Plugins/Mutable/Source/CustomizableObject/Private/MuCO/CustomizableObjectModule.cpp
--- a/Engine/Plugins/Mutable/Source/CustomizableObject/Private/MuCO/CustomizableObjectModule.cpp
+++ b/Engine/Plugins/Mutable/Source/CustomizableObject/Private/MuCO/CustomizableObjectModule.cpp
@@ -241,6 +241,11 @@ void FCustomizableObjectModule::RegisterExternalOperation(const UScriptStruct* O
 
 void FCustomizableObjectModule::UnregisterExternalOperation(const UScriptStruct* Operation)
 {
+	if (IsEngineExitRequested())
+	{
+		return;
+	}
+	
 	if (!Operation)
 	{
 		return;
diff --git a/Engine/Plugins/Mutable/Source/CustomizableObjectEditor/Private/MuCOE/CustomizableObjectEditorModule.cpp b/Engine/Plugins/Mutable/Source/CustomizableObjectEditor/Private/MuCOE/CustomizableObjectEditorModule.cpp
--- a/Engine/Plugins/Mutable/Source/CustomizableObjectEditor/Private/MuCOE/CustomizableObjectEditorModule.cpp
+++ b/Engine/Plugins/Mutable/Source/CustomizableObjectEditor/Private/MuCOE/CustomizableObjectEditorModule.cpp
@@ -1236,6 +1236,11 @@ void FCustomizableObjectEditorModule::RegisterExternalOperation(const UScriptStr
 
 void FCustomizableObjectEditorModule::UnregisterExternalOperation(const UScriptStruct* Operation)
 {
+	if (IsEngineExitRequested())
+	{
+		return;
+	}
+	
 	if (!Operation)
 	{
 		return;
