[Mutable] Avoid crashing if the PinType is None.

#Synced-CL 51468966

diff --git a/Engine/Plugins/Mutable/Source/CustomizableObjectEditor/Private/MuCOE/EdGraphSchema_CustomizableObject.cpp b/Engine/Plugins/Mutable/Source/CustomizableObjectEditor/Private/MuCOE/EdGraphSchema_CustomizableObject.cpp
--- a/Engine/Plugins/Mutable/Source/CustomizableObjectEditor/Private/MuCOE/EdGraphSchema_CustomizableObject.cpp
+++ b/Engine/Plugins/Mutable/Source/CustomizableObjectEditor/Private/MuCOE/EdGraphSchema_CustomizableObject.cpp
@@ -931,7 +931,6 @@ FLinearColor UEdGraphSchema_CustomizableObject::GetPinTypeColor(const FName& Typ
 		return *Result;
 	}
 
-	unimplemented();
 	return FLinearColor::Black;
 }
 
