[Mutable] Fix Extension Data not being applied when an Instance is already generated.

#Synced-CL 51468966

diff --git a/Engine/Plugins/Experimental/HairStrandsMutable/Source/HairStrandsMutable/Private/HairStrandsMutableExtension.cpp b/Engine/Plugins/Experimental/HairStrandsMutable/Source/HairStrandsMutable/Private/HairStrandsMutableExtension.cpp
--- a/Engine/Plugins/Experimental/HairStrandsMutable/Source/HairStrandsMutable/Private/HairStrandsMutableExtension.cpp
+++ b/Engine/Plugins/Experimental/HairStrandsMutable/Source/HairStrandsMutable/Private/HairStrandsMutableExtension.cpp
@@ -43,7 +43,7 @@ TArray<FObjectNodeInputPin> UHairStrandsMutableExtension::GetAdditionalObjectNod
 	return Result;
 }
 
-void UHairStrandsMutableExtension::OnCustomizableObjectInstanceUsageUpdated(UCustomizableObjectInstanceUsage& InstanceUsage, TArray<const UObject*>& ExtensionData) const
+void UHairStrandsMutableExtension::OnCustomizableObjectInstanceUsageUpdated(UCustomizableObjectInstanceUsage& InstanceUsage, TArray<TObjectPtr<const UObject>>& ExtensionData) const
 {
 	Super::OnCustomizableObjectInstanceUsageUpdated(InstanceUsage, ExtensionData);
 	
diff --git a/Engine/Plugins/Experimental/HairStrandsMutable/Source/HairStrandsMutable/Public/HairStrandsMutableExtension.h b/Engine/Plugins/Experimental/HairStrandsMutable/Source/HairStrandsMutable/Public/HairStrandsMutableExtension.h
--- a/Engine/Plugins/Experimental/HairStrandsMutable/Source/HairStrandsMutable/Public/HairStrandsMutableExtension.h
+++ b/Engine/Plugins/Experimental/HairStrandsMutable/Source/HairStrandsMutable/Public/HairStrandsMutableExtension.h
@@ -105,7 +105,7 @@ public:
 	/** UCustomizableObjectExtension interface */
 	virtual TArray<FCustomizableObjectPinType> GetPinTypes() const override;
 	virtual TArray<FObjectNodeInputPin> GetAdditionalObjectNodePins() const override;
-	virtual void OnCustomizableObjectInstanceUsageUpdated(UCustomizableObjectInstanceUsage& Usage, TArray<const UObject*>& ExtensionData) const override;
+	virtual void OnCustomizableObjectInstanceUsageUpdated(UCustomizableObjectInstanceUsage& Usage, TArray<TObjectPtr<const UObject>>& ExtensionData) const override;
 	virtual void OnCustomizableObjectInstanceUsageDiscarded(UCustomizableObjectInstanceUsage& Usage) const override;
 
 	HAIRSTRANDSMUTABLE_API static const FName GroomPinType;
diff --git a/Engine/Plugins/Mutable/Source/CustomizableObject/Internal/MuCO/CustomizableObjectInstancePrivate.h b/Engine/Plugins/Mutable/Source/CustomizableObject/Internal/MuCO/CustomizableObjectInstancePrivate.h
--- a/Engine/Plugins/Mutable/Source/CustomizableObject/Internal/MuCO/CustomizableObjectInstancePrivate.h
+++ b/Engine/Plugins/Mutable/Source/CustomizableObject/Internal/MuCO/CustomizableObjectInstancePrivate.h
@@ -177,6 +177,10 @@ public:
 	UPROPERTY(Transient, VisibleAnywhere, Category = NoCategory)
 	TArray<TObjectPtr<UTexture>> Textures; // TODO GMT Should be used for display only
 
+	/** Generated Extension Data. */
+	UPROPERTY(Transient)
+	TArray<TObjectPtr<const UObject>> ExtensionData;
+	
 	/** If check, Mutable will keep the ownership of the generated resources. This means that Mutable may decide to modify or reuse them for performance considerations. */
 	UPROPERTY(EditAnywhere, Category = UpdateOptions)
 	bool bKeepOwnershipOfGeneratedResources = true;
diff --git a/Engine/Plugins/Mutable/Source/CustomizableObject/Private/MuCO/CustomizableObjectInstance.cpp b/Engine/Plugins/Mutable/Source/CustomizableObject/Private/MuCO/CustomizableObjectInstance.cpp
--- a/Engine/Plugins/Mutable/Source/CustomizableObject/Private/MuCO/CustomizableObjectInstance.cpp
+++ b/Engine/Plugins/Mutable/Source/CustomizableObject/Private/MuCO/CustomizableObjectInstance.cpp
@@ -114,6 +114,7 @@ void UCustomizableInstancePrivate::InvalidateGeneratedData()
 	SkeletalMeshes.Empty();
 	Materials.Empty();
 	Textures.Empty();
+	ExtensionData.Empty();
 	
 	CommittedDescriptor = {};
 	CommittedDescriptorHash = {};
diff --git a/Engine/Plugins/Mutable/Source/CustomizableObject/Private/MuCO/CustomizableObjectSystem.cpp b/Engine/Plugins/Mutable/Source/CustomizableObject/Private/MuCO/CustomizableObjectSystem.cpp
--- a/Engine/Plugins/Mutable/Source/CustomizableObject/Private/MuCO/CustomizableObjectSystem.cpp
+++ b/Engine/Plugins/Mutable/Source/CustomizableObject/Private/MuCO/CustomizableObjectSystem.cpp
@@ -837,16 +837,6 @@ void FinishUpdateGlobal(const TSharedRef<FUpdateContextPrivate>& Context)
 	if (UCustomizableObjectSystem::IsUpdateResultValid(Context->UpdateResult))
 	{
 		MUTABLE_CPUPROFILER_SCOPE(FinishUpdateGlobal_Callbacks);
-
-		TArray<const UObject*> ExtensionData;
-		ExtensionData.Reserve(Context->ExtensionData.Num());
-		for (UE::Mutable::Private::TManagedPtr<const UE::Mutable::Private::FExtensionData>& Data : Context->ExtensionData)
-		{
-			if (UObject* PassthroughObject = Data->PassthroughObject.Get())
-			{
-				ExtensionData.Add(PassthroughObject);
-			}
-		}
 		
 		// Call CustomizableObjectInstanceUsages updated callbacks.
 		if (Context->bOptimizedUpdate)
@@ -866,11 +856,6 @@ void FinishUpdateGlobal(const TSharedRef<FUpdateContextPrivate>& Context)
 				}
 #endif
 				
-				for (const UCustomizableObjectExtension* Extension : ICustomizableObjectModule::Get().GetRegisteredExtensions())
-				{
-					Extension->OnCustomizableObjectInstanceUsageUpdated(*InstanceUsage.Get(), ExtensionData);
-				}
-				
 				InstanceUsage->GetPrivate()->Callbacks();
 			}
 		}
@@ -887,11 +872,6 @@ void FinishUpdateGlobal(const TSharedRef<FUpdateContextPrivate>& Context)
 
 				if (InstanceUsage->GetCustomizableObjectInstance() == Instance)
 				{
-					for (const UCustomizableObjectExtension* Extension : ICustomizableObjectModule::Get().GetRegisteredExtensions())
-					{
-						Extension->OnCustomizableObjectInstanceUsageUpdated(*InstanceUsage, ExtensionData);
-					}
-					
 					InstanceUsage->GetPrivate()->Callbacks();
 				}
 			}
@@ -1072,8 +1052,8 @@ void UpdateSkeletalMesh(const TSharedRef<FUpdateContextPrivate>& Context)
 	
 	uint32 ParentsWithOverrideMaterialsUpdatedCount = 0;
 	
-	UCustomizableInstancePrivate* CustomizableObjectInstancePrivateData = CustomizableObjectInstance->GetPrivate();
-	check(CustomizableObjectInstancePrivateData != nullptr);
+	UCustomizableInstancePrivate* InstancePrivate = CustomizableObjectInstance->GetPrivate();
+	check(InstancePrivate != nullptr);
 
 	for (UCustomizableObjectInstanceUsage* Usage : UCustomizableObjectSystemPrivate::GetCustomizableObjectInstanceUsages())
 	{
@@ -1165,7 +1145,7 @@ void UpdateSkeletalMesh(const TSharedRef<FUpdateContextPrivate>& Context)
 			}
 		}
 		
-		if (CustomizableObjectInstancePrivateData->HasCOInstanceFlags(ReplacePhysicsAssets) &&
+		if (InstancePrivate->HasCOInstanceFlags(ReplacePhysicsAssets) &&
 			SkeletalMesh &&
 			Parent->GetWorld())
 		{
@@ -1181,6 +1161,11 @@ void UpdateSkeletalMesh(const TSharedRef<FUpdateContextPrivate>& Context)
 		{
 			Context->AttachedParentUpdated.Add(Usage);
 		}
+		
+		for (const UCustomizableObjectExtension* Extension : ICustomizableObjectModule::Get().GetRegisteredExtensions())
+		{
+			Extension->OnCustomizableObjectInstanceUsageUpdated(*Usage, InstancePrivate->ExtensionData);
+		}
 	}
 
 	if (ParentsWithOverrideMaterialsUpdatedCount > 0)
@@ -2305,14 +2290,26 @@ namespace impl
 			Subtask_Mutable_PrepareRealTimeMorphData(OperationData);
 		}
 
-		UCustomizableObjectInstance* CustomizableObjectInstance = OperationData->Instance.Get();
+		UCustomizableObjectInstance* Instance = OperationData->Instance.Get();
 
 		// Actual work
-		const bool bInstanceInvalid = !IsValid(CustomizableObjectInstance);
+		const bool bInstanceInvalid = !IsValid(Instance);
 		if (!bInstanceInvalid)
 		{
-			UCustomizableInstancePrivate* CustomizableInstancePrivateData = CustomizableObjectInstance->GetPrivate();
+			UCustomizableInstancePrivate* InstancePrivate = Instance->GetPrivate();
 
+			// Extension data
+			{
+				InstancePrivate->ExtensionData.Empty(OperationData->ExtensionData.Num());
+				for (UE::Mutable::Private::TManagedPtr<const UE::Mutable::Private::FExtensionData>& Data : OperationData->ExtensionData)
+				{
+					if (UObject* PassthroughObject = Data->PassthroughObject.Get())
+					{
+						InstancePrivate->ExtensionData.Add(PassthroughObject);
+					}
+				}
+			}
+			
 			// Prepare the component data
 			//-------------------------------------------------------------
 			{
@@ -2322,7 +2319,7 @@ namespace impl
 				const UModelResources& ModelResources = Object->GetPrivate()->GetModelResourcesChecked();
 				
 				const int32 NumComponents = OperationData->MutableInstance->GetComponentCount();
-				CustomizableInstancePrivateData->ComponentsData.SetNum(NumComponents);
+				InstancePrivate->ComponentsData.SetNum(NumComponents);
 
 				TArray<FName> ComponentNames;
 				ComponentNames.Reserve(NumComponents);
@@ -2344,8 +2341,8 @@ namespace impl
 
 					ComponentNames.Add(ComponentName);
 
-					CustomizableInstancePrivateData->ComponentsData[ComponentIndex].ComponentId = ComponentId;
-					CustomizableInstancePrivateData->ComponentsData[ComponentIndex].ComponentName = ComponentName;
+					InstancePrivate->ComponentsData[ComponentIndex].ComponentId = ComponentId;
+					InstancePrivate->ComponentsData[ComponentIndex].ComponentName = ComponentName;
 				}
 			}
 			
@@ -2353,14 +2350,14 @@ namespace impl
 			//-------------------------------------------------------------
 
 			// \TODO: Bring that code here instead of keeping it in the UCustomizableObjectInstance
-			if (CustomizableInstancePrivateData->UpdateSkeletalMesh_PostBeginUpdate0(CustomizableObjectInstance, OperationData))
+			if (InstancePrivate->UpdateSkeletalMesh_PostBeginUpdate0(Instance, OperationData))
 			{
 				// This used to be CustomizableObjectInstance::UpdateSkeletalMesh_PostBeginUpdate1
 				{
 					MUTABLE_CPUPROFILER_SCOPE(UpdateSkeletalMesh_PostBeginUpdate1);
 
 					// \TODO: Bring here
-					CustomizableInstancePrivateData->BuildResources(OperationData, CustomizableObjectInstance);
+					InstancePrivate->BuildResources(OperationData, Instance);
 				}
 
 				// This used to be CustomizableObjectInstance::UpdateSkeletalMesh_PostBeginUpdate2
@@ -2368,9 +2365,9 @@ namespace impl
 					MUTABLE_CPUPROFILER_SCOPE(UpdateSkeletalMesh_PostBeginUpdate2);
 					
 #if WITH_EDITORONLY_DATA
-					CustomizableInstancePrivateData->RegenerateImportedModels(OperationData);
+					InstancePrivate->RegenerateImportedModels(OperationData);
 #endif
-					CustomizableInstancePrivateData->PostEditChangePropertyWithoutEditor();
+					InstancePrivate->PostEditChangePropertyWithoutEditor();
 				}
 			}
 		} // if (!bInstanceValid)
diff --git a/Engine/Plugins/Mutable/Source/CustomizableObject/Private/MuCO/CustomizableSkeletalMeshActor.cpp b/Engine/Plugins/Mutable/Source/CustomizableObject/Private/MuCO/CustomizableSkeletalMeshActor.cpp
--- a/Engine/Plugins/Mutable/Source/CustomizableObject/Private/MuCO/CustomizableSkeletalMeshActor.cpp
+++ b/Engine/Plugins/Mutable/Source/CustomizableObject/Private/MuCO/CustomizableSkeletalMeshActor.cpp
@@ -76,42 +76,36 @@ void UCustomizableSkeletalMeshActorPrivate::Init(UCustomizableObjectInstance* In
 		return;
 	}
 	
-	for (UE::Mutable::Private::FComponentId ComponentId = 0; ComponentId < Object->GetComponentCount(); ++ComponentId)
+	const int32 NumComponents = Object->GetComponentCount();
+	
+	for (TObjectPtr<USkeletalMeshComponent> Component : GetPublic()->SkeletalMeshComponents)
+	{
+		Component->UnregisterComponent();
+	}
+	
+	GetPublic()->SkeletalMeshComponents.Empty(NumComponents);
+
+	for (TObjectPtr<UCustomizableSkeletalComponent> Component : GetPublic()->CustomizableSkeletalComponents)
+	{
+		Component->UnregisterComponent();
+	}
+
+	GetPublic()->CustomizableSkeletalComponents.Empty(NumComponents);
+	
+	for (UE::Mutable::Private::FComponentId ComponentId = 0; ComponentId < NumComponents; ++ComponentId)
 	{
 		const FName& ComponentName = Object->GetComponentName(ComponentId);
 		const FString ComponentNameString = ComponentName.ToString();
 
-		USkeletalMeshComponent* Component;
-		if (ComponentId == 0)
-		{
-			Component = GetPublic()->SkeletalMeshComponents[0];
-			// Renaming the component will make it disappear. Keep it as it is.
-		}
-		else
-		{
-			const FString SkeletalMeshComponentName = FString::Printf(TEXT("SkeletalMeshComponent %s"), *ComponentNameString);
-			Component = NewObject<USkeletalMeshComponent>(GetPublic(), USkeletalMeshComponent::StaticClass(), FName(*SkeletalMeshComponentName));
-			Component->CreationMethod = EComponentCreationMethod::Native;
-			Component->AttachToComponent(GetPublic()->GetRootComponent(), FAttachmentTransformRules::KeepRelativeTransform);
-			Component->RegisterComponent();
-		}
-					
-		if (!Component)
-		{
-			continue;
-		}
-
+		const FString SkeletalMeshComponentName = FString::Printf(TEXT("SkeletalMeshComponent %s"), *ComponentNameString);
+		USkeletalMeshComponent* Component = NewObject<USkeletalMeshComponent>(GetPublic(), USkeletalMeshComponent::StaticClass(), FName(*SkeletalMeshComponentName));
+		Component->CreationMethod = EComponentCreationMethod::Native;
+		Component->AttachToComponent(GetPublic()->GetRootComponent(), FAttachmentTransformRules::KeepRelativeTransform);
+		Component->RegisterComponent();
+		
 		UCustomizableSkeletalComponent* CustomizableComponent;
-		if (ComponentId == 0)
-		{
-			CustomizableComponent = GetPublic()->CustomizableSkeletalComponents[0];
-			// Renaming the component will make it disappear. Keep it as it is.
-		}
-		else
-		{
-			const FString CustomizableComponentName = FString::Printf(TEXT("CustomizableSkeletalComponent %s"), *ComponentNameString);
-			CustomizableComponent = NewObject<UCustomizableSkeletalComponent>(GetPublic(), UCustomizableSkeletalComponent::StaticClass(), FName(CustomizableComponentName));
-		}
+		const FString CustomizableComponentName = FString::Printf(TEXT("CustomizableSkeletalComponent %s"), *ComponentNameString);
+		CustomizableComponent = NewObject<UCustomizableSkeletalComponent>(GetPublic(), UCustomizableSkeletalComponent::StaticClass(), FName(CustomizableComponentName));
 
 		CustomizableComponent->AttachToComponent(Component, FAttachmentTransformRules::KeepRelativeTransform);
 		CustomizableComponent->CustomizableObjectInstance = Instance;
diff --git a/Engine/Plugins/Mutable/Source/CustomizableObject/Public/MuCO/CustomizableObjectExtension.h b/Engine/Plugins/Mutable/Source/CustomizableObject/Public/MuCO/CustomizableObjectExtension.h
--- a/Engine/Plugins/Mutable/Source/CustomizableObject/Public/MuCO/CustomizableObjectExtension.h
+++ b/Engine/Plugins/Mutable/Source/CustomizableObject/Public/MuCO/CustomizableObjectExtension.h
@@ -102,7 +102,7 @@ public:
 	/** Called when a Customizable Object Instance Usage is being updated.
 	 * 
 	 * @param Usage Customizable Object Instance Usage being updated. */
-	virtual void OnCustomizableObjectInstanceUsageUpdated(UCustomizableObjectInstanceUsage& Usage, TArray<const UObject*>& ExtensionData) const {}
+	virtual void OnCustomizableObjectInstanceUsageUpdated(UCustomizableObjectInstanceUsage& Usage, TArray<TObjectPtr<const UObject>>& ExtensionData) const {}
 
 	/** Called when a Customizable Object Instance Usage is being discarded.
 	 * 
diff --git a/Engine/Plugins/Mutable/Source/CustomizableObject/Public/MuCO/CustomizableSkeletalMeshActor.h b/Engine/Plugins/Mutable/Source/CustomizableObject/Public/MuCO/CustomizableSkeletalMeshActor.h
--- a/Engine/Plugins/Mutable/Source/CustomizableObject/Public/MuCO/CustomizableSkeletalMeshActor.h
+++ b/Engine/Plugins/Mutable/Source/CustomizableObject/Public/MuCO/CustomizableSkeletalMeshActor.h
@@ -42,8 +42,6 @@ public:
 	UE_API const UCustomizableSkeletalMeshActorPrivate* GetPrivate() const;
 
 private:
-	UE_API void Init(UCustomizableObjectInstance* Instance);
-	
 	// DEPRECATED
 	UFUNCTION()
 	UE_API void SwitchComponentsMaterials(UCustomizableObjectInstance* Instance);
diff --git a/Engine/Plugins/Mutable/Source/CustomizableObjectEditor/Private/MuCOE/CustomizableObjectInstanceFactory.cpp b/Engine/Plugins/Mutable/Source/CustomizableObjectEditor/Private/MuCOE/CustomizableObjectInstanceFactory.cpp
--- a/Engine/Plugins/Mutable/Source/CustomizableObjectEditor/Private/MuCOE/CustomizableObjectInstanceFactory.cpp
+++ b/Engine/Plugins/Mutable/Source/CustomizableObjectEditor/Private/MuCOE/CustomizableObjectInstanceFactory.cpp
@@ -27,29 +27,31 @@ UCustomizableObjectInstanceFactory::UCustomizableObjectInstanceFactory(const FOb
 }
 
 
-void UCustomizableObjectInstanceFactory::PostSpawnActor(UObject* Asset, AActor* NewActor)
+AActor* UCustomizableObjectInstanceFactory::SpawnActor(UObject* InAsset, ULevel* InLevel, const FTransform& InTransform, const FActorSpawnParameters& InSpawnParams)
 {
-	Super::PostSpawnActor(Asset, NewActor);
+	AActor* NewActor = Super::SpawnActor(InAsset, InLevel, InTransform, InSpawnParams);
 	
-	UCustomizableObjectInstance* Instance = Cast<UCustomizableObjectInstance>(Asset);
+	ACustomizableSkeletalMeshActor* NewCSMActor = Cast<ACustomizableSkeletalMeshActor>(NewActor);
+    if (!NewCSMActor)
+    {
+    	return nullptr;
+    }
+	
+	UCustomizableObjectInstance* Instance = Cast<UCustomizableObjectInstance>(InAsset);
 	if (!Instance)
 	{
-		return;
+		return nullptr;
 	}
 
 	UCustomizableObject* Object = Instance->GetCustomizableObject();
 	if (!Object)
 	{
-		return;
+		return nullptr;
 	}
-
-	ACustomizableSkeletalMeshActor* NewCSMActor = CastChecked<ACustomizableSkeletalMeshActor>(NewActor);
-	if (!NewCSMActor)
-	{
-		return;
-	}
-
+	
 	NewCSMActor->GetPrivate()->Init(Instance);
+	
+	return NewActor;
 }
 
 
diff --git a/Engine/Plugins/Mutable/Source/CustomizableObjectEditor/Public/MuCOE/CustomizableObjectInstanceFactory.h b/Engine/Plugins/Mutable/Source/CustomizableObjectEditor/Public/MuCOE/CustomizableObjectInstanceFactory.h
--- a/Engine/Plugins/Mutable/Source/CustomizableObjectEditor/Public/MuCOE/CustomizableObjectInstanceFactory.h
+++ b/Engine/Plugins/Mutable/Source/CustomizableObjectEditor/Public/MuCOE/CustomizableObjectInstanceFactory.h
@@ -18,7 +18,7 @@ class UCustomizableObjectInstanceFactory : public UActorFactory
 
 protected:
 	// UActorFactory interface
-	virtual void PostSpawnActor(UObject* Asset, AActor* NewActor) override;
+	virtual AActor* SpawnActor(UObject* InAsset, ULevel* InLevel, const FTransform& InTransform, const FActorSpawnParameters& InSpawnParams) override;
 	virtual UObject* GetAssetFromActorInstance(AActor* ActorInstance) override;
 	virtual bool CanCreateActorFrom(const FAssetData& AssetData, FText& OutErrorMsg) override;
 	virtual FQuat AlignObjectToSurfaceNormal(const FVector& InSurfaceNormal, const FQuat& ActorRotation) const override;
