[Mutable] Remove manual LODs

#rnx

#Synced-CL 51492284

diff --git a/Engine/Plugins/Mutable/Source/CustomizableObjectEditor/Private/MuCOE/GenerateMutableSource/GenerateMutableSource.cpp b/Engine/Plugins/Mutable/Source/CustomizableObjectEditor/Private/MuCOE/GenerateMutableSource/GenerateMutableSource.cpp
--- a/Engine/Plugins/Mutable/Source/CustomizableObjectEditor/Private/MuCOE/GenerateMutableSource/GenerateMutableSource.cpp
+++ b/Engine/Plugins/Mutable/Source/CustomizableObjectEditor/Private/MuCOE/GenerateMutableSource/GenerateMutableSource.cpp
@@ -991,13 +991,11 @@ UE::Mutable::Private::Ptr<UE::Mutable::Private::NodeObject> GenerateMutableSourc
 				for (const FSkeletalMeshComponentInfo& Pair : GenerationContext.CompilationContext->ComponentInfos)
 				{
 					GenerationContext.CurrentSkeletalMeshComponent = Pair.ComponentId;
-					GenerationContext.CurrentAutoLODStrategy = Pair.Node->AutoLODStrategy;
 					
 					UE::Mutable::Private::Ptr<UE::Mutable::Private::NodeModifier> ModifierNode = GenerateMutableSourceModifier(ConnectedModifier, GenerationContext);
 					ObjectNode->Modifiers.Add(ModifierNode);
 
 					GenerationContext.CurrentSkeletalMeshComponent = INDEX_NONE;
-					GenerationContext.CurrentAutoLODStrategy = {};
 				}
 			}
 		}
diff --git a/Engine/Plugins/Mutable/Source/CustomizableObjectEditor/Private/MuCOE/GenerateMutableSource/GenerateMutableSource.h b/Engine/Plugins/Mutable/Source/CustomizableObjectEditor/Private/MuCOE/GenerateMutableSource/GenerateMutableSource.h
--- a/Engine/Plugins/Mutable/Source/CustomizableObjectEditor/Private/MuCOE/GenerateMutableSource/GenerateMutableSource.h
+++ b/Engine/Plugins/Mutable/Source/CustomizableObjectEditor/Private/MuCOE/GenerateMutableSource/GenerateMutableSource.h
@@ -842,9 +842,6 @@ struct FMutableGraphGenerationContext
 	TMap<UE::Mutable::Private::FComponentId, uint8> NumMaxLODsToStream;
 
 	bool bPartialCompilation = false;
-	
-	// Based on the last object visited.
-	ECustomizableObjectAutomaticLODStrategy CurrentAutoLODStrategy = ECustomizableObjectAutomaticLODStrategy::Manual;
 
 	// Stores external graph root nodes to be added to the specified group nodes
 	TMultiMap<FGuid, UCustomizableObjectNodeObject*> GroupIdToExternalNodeMap;
diff --git a/Engine/Plugins/Mutable/Source/CustomizableObjectEditor/Private/MuCOE/GenerateMutableSource/GenerateMutableSourceComponent.cpp b/Engine/Plugins/Mutable/Source/CustomizableObjectEditor/Private/MuCOE/GenerateMutableSource/GenerateMutableSourceComponent.cpp
--- a/Engine/Plugins/Mutable/Source/CustomizableObjectEditor/Private/MuCOE/GenerateMutableSource/GenerateMutableSourceComponent.cpp
+++ b/Engine/Plugins/Mutable/Source/CustomizableObjectEditor/Private/MuCOE/GenerateMutableSource/GenerateMutableSourceComponent.cpp
@@ -83,12 +83,10 @@ UE::Mutable::Private::Ptr<UE::Mutable::Private::NodeComponent> GenerateMutableSo
 			Result = NodeComponentNew;
 
 			GenerationContext.CurrentSkeletalMeshComponent = ComponentInfo->ComponentId;
-			GenerationContext.CurrentAutoLODStrategy = TypedComponentMesh->AutoLODStrategy;
 		
 			NodeComponentNew->SkeletalMesh = UE::Mutable::Private::GenerateMutableSourceSkeletalMesh(Pin, GenerationContext);
 		
 			GenerationContext.CurrentSkeletalMeshComponent = INDEX_NONE;
-			GenerationContext.CurrentAutoLODStrategy = {};
 		}
 	}
 
@@ -123,15 +121,10 @@ UE::Mutable::Private::Ptr<UE::Mutable::Private::NodeComponent> GenerateMutableSo
 			UE::Mutable::Private::Ptr<UE::Mutable::Private::NodeComponentEdit> NodeComponentEdit = new UE::Mutable::Private::NodeComponentEdit();
 			NodeComponentEdit->Parent = ParentNodeComponent.get();
 			NodeComponentEdit->SetMessageContext(TypedComponentMeshExtend);
-		
-			GenerationContext.CurrentAutoLODStrategy = TypedComponentMeshExtend->AutoLODStrategy == ECustomizableObjectAutomaticLODStrategy::Inherited ?
-				TypedParentComponentMesh->AutoLODStrategy :
-				TypedComponentMeshExtend->AutoLODStrategy;
 			
 			NodeComponentEdit->SkeletalMesh = UE::Mutable::Private::GenerateMutableSourceSkeletalMeshNew(Pin, GenerationContext);
 
 			GenerationContext.CurrentSkeletalMeshComponent = INDEX_NONE;
-			GenerationContext.CurrentAutoLODStrategy = {};
 			
 			Result = NodeComponentEdit;	
 		}
@@ -397,13 +390,6 @@ void FirstPass(UCustomizableObjectNodeComponentMesh& Node, FMutableGraphGenerati
 		return;
 	}
 	
-	// Ensure that the CO has a valid AutoLODStrategy on the Component node.
-	if (Node.AutoLODStrategy == ECustomizableObjectAutomaticLODStrategy::Inherited)
-	{
-		GenerationContext.Log(LOCTEXT("RootInheritsFromParent", "Error! Component LOD Strategy can't be set to 'Inherit from parent object'"), &Node, EMessageSeverity::Error);
-		return;
-	}
-
 	// Fill the basic LOD Settings
 	uint8 NumLODs = Node.LODPins.Num();
 
diff --git a/Engine/Plugins/Mutable/Source/CustomizableObjectEditor/Private/MuCOE/GenerateMutableSource/GenerateMutableSourceMesh.cpp b/Engine/Plugins/Mutable/Source/CustomizableObjectEditor/Private/MuCOE/GenerateMutableSource/GenerateMutableSourceMesh.cpp
--- a/Engine/Plugins/Mutable/Source/CustomizableObjectEditor/Private/MuCOE/GenerateMutableSource/GenerateMutableSourceMesh.cpp
+++ b/Engine/Plugins/Mutable/Source/CustomizableObjectEditor/Private/MuCOE/GenerateMutableSource/GenerateMutableSourceMesh.cpp
@@ -229,7 +229,7 @@ void SetBuffersFormat(FMutableCompilationContext& Context,
 }
 
 
-void GetLODAndSectionForAutomaticLODs(const FMutableCompilationContext& Context, const FMutableSourceMeshData& Source, const USkeletalMesh& SkeletalMesh, int32& OutLODIndex, int32& OutSectionIndex)
+void GetLODAndSection(const FMutableCompilationContext& Context, const FMutableSourceMeshData& Source, const USkeletalMesh& SkeletalMesh, int32& OutLODIndex, int32& OutSectionIndex)
 {
 	OutLODIndex = Source.BaseLODIndex;
 	OutSectionIndex = Source.BaseSectionIndex;
@@ -990,7 +990,7 @@ UE::Tasks::FTask ConvertSkeletalMeshToMutable(TSharedRef<FMeshConversionContext>
 	// Find the correct LOD and Section indices when using automatic LODs 
 	int32 LODIndex = INDEX_NONE;
 	int32 SectionIndex = INDEX_NONE;
-	GetLODAndSectionForAutomaticLODs(Context, Source, *SkeletalMesh, LODIndex, SectionIndex);
+	GetLODAndSection(Context, Source, *SkeletalMesh, LODIndex, SectionIndex);
 
 	if (!ImportedModel->LODModels.IsValidIndex(LODIndex))
 	{
@@ -2630,10 +2630,9 @@ UE::Mutable::Private::Ptr<UE::Mutable::Private::NodeMesh> GenerateMutableSourceM
 			
 			// Layouts
 			// 
-			// Generate Layouts when using manual LODs or when processing the base Automatic LOD (the other automatic LODs will use the base layouts).
-			const bool bGenerateLayouts = 
-					GenerationContext.CurrentAutoLODStrategy != ECustomizableObjectAutomaticLODStrategy::AutomaticFromMesh || 
-					(GenerationContext.CurrentLOD == GenerationContext.FromLOD || GenerationContext.CurrentLOD <= GenerationContext.FirstLODAvailable[GenerationContext.CurrentSkeletalMeshComponent]);
+			// Generate Layouts when processing the first section using this layout (the other LODs will use the base layouts).
+			const bool bGenerateLayouts = GenerationContext.CurrentLOD == GenerationContext.FromLOD ||
+				GenerationContext.CurrentLOD <= GenerationContext.FirstLODAvailable[GenerationContext.CurrentSkeletalMeshComponent];
 
 			if (bGenerateLayouts)
 			{
@@ -3440,7 +3439,7 @@ UE::Mutable::Private::Ptr<UE::Mutable::Private::NodeMesh> GenerateMutableSourceM
 							Source.LODOffset = GenerationContext.CurrentLOD - GenerationContext.FromLOD;
 							Source.MessageContext = Node;
 							Source.bOnlyConnectedLOD = bOnlyConnectedLOD;
-							GetLODAndSectionForAutomaticLODs(*GenerationContext.CompilationContext, Source, *DefaultSkeletalMesh, LODIndex, SectionIndex);
+							GetLODAndSection(*GenerationContext.CompilationContext, Source, *DefaultSkeletalMesh, LODIndex, SectionIndex);
 
 							MutableColumnName = TypedNodeTable->GenerateSkeletalMeshMutableColumName(ColumnName, LODIndex, SectionIndex);	
 						}
@@ -3487,9 +3486,9 @@ UE::Mutable::Private::Ptr<UE::Mutable::Private::NodeMesh> GenerateMutableSourceM
 
 							// Layouts
 							// 
-							// Generate Layouts when using manual LODs or when processing the base Automatic LOD (the other automatic LODs will use the base layouts).
-							const bool bGenerateLayouts = GenerationContext.CurrentAutoLODStrategy != ECustomizableObjectAutomaticLODStrategy::AutomaticFromMesh
-								|| (GenerationContext.CurrentLOD == GenerationContext.FromLOD || GenerationContext.CurrentLOD <= GenerationContext.FirstLODAvailable[GenerationContext.CurrentSkeletalMeshComponent]);
+							//  Generate Layouts when processing the first section using this layout(the other LODs will use the base layouts).
+							const bool bGenerateLayouts = GenerationContext.CurrentLOD == GenerationContext.FromLOD || 
+								GenerationContext.CurrentLOD <= GenerationContext.FirstLODAvailable[GenerationContext.CurrentSkeletalMeshComponent];
 
 							if (bGenerateLayouts)
 							{
@@ -3561,7 +3560,6 @@ UE::Mutable::Private::Ptr<UE::Mutable::Private::NodeMesh> GenerateMutableSourceM
 			TypedNodeParam->GetPinSection(*Pin, BaseLODIndex, BaseSectionIndex);
 			BreakNode->LODIndex = BaseLODIndex;
 			BreakNode->SectionIndex = BaseSectionIndex;
-			BreakNode->bAutomaticLODs = GenerationContext.CurrentAutoLODStrategy == ECustomizableObjectAutomaticLODStrategy::AutomaticFromMesh;
 		}
 		
 		BreakNode->ConversionFlags = (uint8)GenerationContext.MeshGenerationFlags.Last();
diff --git a/Engine/Plugins/Mutable/Source/CustomizableObjectEditor/Private/MuCOE/GenerateMutableSource/GenerateMutableSourceMesh.h b/Engine/Plugins/Mutable/Source/CustomizableObjectEditor/Private/MuCOE/GenerateMutableSource/GenerateMutableSourceMesh.h
--- a/Engine/Plugins/Mutable/Source/CustomizableObjectEditor/Private/MuCOE/GenerateMutableSource/GenerateMutableSourceMesh.h
+++ b/Engine/Plugins/Mutable/Source/CustomizableObjectEditor/Private/MuCOE/GenerateMutableSource/GenerateMutableSourceMesh.h
@@ -26,7 +26,7 @@ namespace UE::Mutable::Private
 }
 
 
-/** Returns the corrected LOD and Section Index when using Automatic LOD From Mesh strategy.
+/** Returns the corrected LOD and Section Index
  *
  * Do not confuse Section Index and Material Index, they are not the same.
  * 
@@ -35,8 +35,8 @@ namespace UE::Mutable::Private
  * @param SkeletalMesh 
  * @param OutLODIndex Corrected Skeletal Mesh LOD Index.
  * @param OutSectionIndex Corrected Skeletal Mesh Section Index.
- * @return When using Automatic LOD From Mesh, OutLODIndex and OutSectionIndex will return -1 if the section is not found in the currently compiling LOD. */
-void GetLODAndSectionForAutomaticLODs(const FMutableCompilationContext& Context, const FMutableSourceMeshData& Source, const USkeletalMesh& SkeletalMesh, int32& OutLODIndex, int32& OutSectionIndex);
+ * @return OutLODIndex and OutSectionIndex will return -1 if the section is not found in the currently compiling LOD. */
+void GetLODAndSection(const FMutableCompilationContext& Context, const FMutableSourceMeshData& Source, const USkeletalMesh& SkeletalMesh, int32& OutLODIndex, int32& OutSectionIndex);
 
 
 /** Converts an Unreal Skeletal Mesh to Mutable Mesh.
diff --git a/Engine/Plugins/Mutable/Source/CustomizableObjectEditor/Private/MuCOE/GenerateMutableSource/GenerateMutableSourceSkeletalMesh.cpp b/Engine/Plugins/Mutable/Source/CustomizableObjectEditor/Private/MuCOE/GenerateMutableSource/GenerateMutableSourceSkeletalMesh.cpp
--- a/Engine/Plugins/Mutable/Source/CustomizableObjectEditor/Private/MuCOE/GenerateMutableSource/GenerateMutableSourceSkeletalMesh.cpp
+++ b/Engine/Plugins/Mutable/Source/CustomizableObjectEditor/Private/MuCOE/GenerateMutableSource/GenerateMutableSourceSkeletalMesh.cpp
@@ -91,8 +91,7 @@ namespace UE::Mutable::Private
 
 				const int32 NumLODs = TypedComponentMesh->GetLODPins().Num();
 
-				const bool bUseAutomaticLods = GenerationContext.CurrentAutoLODStrategy == ECustomizableObjectAutomaticLODStrategy::AutomaticFromMesh;
-				FirstLOD = (CurrentLOD < NumLODs) && (FirstLOD == INDEX_NONE || !bUseAutomaticLods) ? CurrentLOD : FirstLOD;
+				FirstLOD = CurrentLOD < NumLODs && FirstLOD == INDEX_NONE ? CurrentLOD : FirstLOD;
 
 				if (FirstLOD < 0)
 				{
diff --git a/Engine/Plugins/Mutable/Source/CustomizableObjectEditor/Private/MuCOE/GenerateMutableSource/GenerateMutableSourceSurface.cpp b/Engine/Plugins/Mutable/Source/CustomizableObjectEditor/Private/MuCOE/GenerateMutableSource/GenerateMutableSourceSurface.cpp
--- a/Engine/Plugins/Mutable/Source/CustomizableObjectEditor/Private/MuCOE/GenerateMutableSource/GenerateMutableSourceSurface.cpp
+++ b/Engine/Plugins/Mutable/Source/CustomizableObjectEditor/Private/MuCOE/GenerateMutableSource/GenerateMutableSourceSurface.cpp
@@ -71,8 +71,7 @@ UE::Mutable::Private::Ptr<UE::Mutable::Private::NodeSurface> GenerateMutableSour
 
 	if (UCONodeSkeletalMeshSectionBase* TypedNodeMat = Cast<UCONodeSkeletalMeshSectionBase>(Node))
 	{
-		if (GenerationContext.CurrentAutoLODStrategy == ECustomizableObjectAutomaticLODStrategy::AutomaticFromMesh
-			&& TypedNodeMat->GetMeshSectionMaxLOD() != INDEX_NONE && TypedNodeMat->GetMeshSectionMaxLOD() < GenerationContext.CurrentLOD)
+		if (TypedNodeMat->GetMeshSectionMaxLOD() != INDEX_NONE && TypedNodeMat->GetMeshSectionMaxLOD() < GenerationContext.CurrentLOD)
 		{
 			return Result;
 		}
diff --git a/Engine/Plugins/Mutable/Source/CustomizableObjectEditor/Private/MuCOE/GenerateMutableSource/GenerateMutableSourceTable.cpp b/Engine/Plugins/Mutable/Source/CustomizableObjectEditor/Private/MuCOE/GenerateMutableSource/GenerateMutableSourceTable.cpp
--- a/Engine/Plugins/Mutable/Source/CustomizableObjectEditor/Private/MuCOE/GenerateMutableSource/GenerateMutableSourceTable.cpp
+++ b/Engine/Plugins/Mutable/Source/CustomizableObjectEditor/Private/MuCOE/GenerateMutableSource/GenerateMutableSourceTable.cpp
@@ -65,7 +65,7 @@ bool FillTableColumn(const UCustomizableObjectNodeTable* TableNode, UE::Mutable:
 				// Current LOD and Section 
 				int32 LODIndex = INDEX_NONE;
 				int32 SectionIndex = INDEX_NONE;
-				GetLODAndSectionForAutomaticLODs(*GenerationContext.CompilationContext, BaseMeshData, *ReferenceSkeletalMesh, LODIndex, SectionIndex);
+				GetLODAndSection(*GenerationContext.CompilationContext, BaseMeshData, *ReferenceSkeletalMesh, LODIndex, SectionIndex);
 
 				MutableColumnName = TableNode->GenerateSkeletalMeshMutableColumName(ColumnName, LODIndex, SectionIndex);
 			}
diff --git a/Engine/Plugins/Mutable/Source/CustomizableObjectEditor/Private/MuCOE/Nodes/CustomizableObjectNodeComponentMesh.cpp b/Engine/Plugins/Mutable/Source/CustomizableObjectEditor/Private/MuCOE/Nodes/CustomizableObjectNodeComponentMesh.cpp
--- a/Engine/Plugins/Mutable/Source/CustomizableObjectEditor/Private/MuCOE/Nodes/CustomizableObjectNodeComponentMesh.cpp
+++ b/Engine/Plugins/Mutable/Source/CustomizableObjectEditor/Private/MuCOE/Nodes/CustomizableObjectNodeComponentMesh.cpp
@@ -154,12 +154,6 @@ int32 UCustomizableObjectNodeComponentMesh::GetNumLODs()
 }
 
 
-ECustomizableObjectAutomaticLODStrategy UCustomizableObjectNodeComponentMesh::GetAutoLODStrategy()
-{
-	return AutoLODStrategy;
-}
-
-
 const TArray<FEdGraphPinReference>& UCustomizableObjectNodeComponentMesh::GetLODPins() const
 {
 	return LODPins;
diff --git a/Engine/Plugins/Mutable/Source/CustomizableObjectEditor/Private/MuCOE/Nodes/CustomizableObjectNodeComponentMesh.h b/Engine/Plugins/Mutable/Source/CustomizableObjectEditor/Private/MuCOE/Nodes/CustomizableObjectNodeComponentMesh.h
--- a/Engine/Plugins/Mutable/Source/CustomizableObjectEditor/Private/MuCOE/Nodes/CustomizableObjectNodeComponentMesh.h
+++ b/Engine/Plugins/Mutable/Source/CustomizableObjectEditor/Private/MuCOE/Nodes/CustomizableObjectNodeComponentMesh.h
@@ -72,7 +72,6 @@ public:
 
 	// ComponentMesh Interface
 	UE_API virtual int32 GetNumLODs() override;
-	UE_API virtual ECustomizableObjectAutomaticLODStrategy GetAutoLODStrategy() override;
 	UE_API virtual const TArray<FEdGraphPinReference>& GetLODPins() const override;
 	UE_API virtual UEdGraphPin* GetOutputPin() const;
 	UE_API virtual void SetOutputPin(const UEdGraphPin* Pin);
@@ -108,9 +107,6 @@ public:
 	UPROPERTY(EditAnywhere, Category = ComponentMesh)
 	int32 NumLODs = 1;
 
-	UPROPERTY(EditAnywhere, Category = ComponentMesh, DisplayName = "Auto LOD Strategy")
-	ECustomizableObjectAutomaticLODStrategy AutoLODStrategy = ECustomizableObjectAutomaticLODStrategy::AutomaticFromMesh;
-
 	UPROPERTY()
 	TArray<FEdGraphPinReference> LODPins;
 
diff --git a/Engine/Plugins/Mutable/Source/CustomizableObjectEditor/Private/MuCOE/Nodes/CustomizableObjectNodeComponentMeshAddTo.cpp b/Engine/Plugins/Mutable/Source/CustomizableObjectEditor/Private/MuCOE/Nodes/CustomizableObjectNodeComponentMeshAddTo.cpp
--- a/Engine/Plugins/Mutable/Source/CustomizableObjectEditor/Private/MuCOE/Nodes/CustomizableObjectNodeComponentMeshAddTo.cpp
+++ b/Engine/Plugins/Mutable/Source/CustomizableObjectEditor/Private/MuCOE/Nodes/CustomizableObjectNodeComponentMeshAddTo.cpp
@@ -145,12 +145,6 @@ int32 UCustomizableObjectNodeComponentMeshAddTo::GetNumLODs()
 }
 
 
-ECustomizableObjectAutomaticLODStrategy UCustomizableObjectNodeComponentMeshAddTo::GetAutoLODStrategy()
-{
-	return AutoLODStrategy;
-}
-
-
 const TArray<FEdGraphPinReference>& UCustomizableObjectNodeComponentMeshAddTo::GetLODPins() const
 {
 	return LODPins;
diff --git a/Engine/Plugins/Mutable/Source/CustomizableObjectEditor/Private/MuCOE/Nodes/CustomizableObjectNodeComponentMeshAddTo.h b/Engine/Plugins/Mutable/Source/CustomizableObjectEditor/Private/MuCOE/Nodes/CustomizableObjectNodeComponentMeshAddTo.h
--- a/Engine/Plugins/Mutable/Source/CustomizableObjectEditor/Private/MuCOE/Nodes/CustomizableObjectNodeComponentMeshAddTo.h
+++ b/Engine/Plugins/Mutable/Source/CustomizableObjectEditor/Private/MuCOE/Nodes/CustomizableObjectNodeComponentMeshAddTo.h
@@ -35,7 +35,6 @@ public:
 	
 	// ComponentMesh Interface
 	UE_API virtual int32 GetNumLODs() override;
-	UE_API virtual ECustomizableObjectAutomaticLODStrategy GetAutoLODStrategy() override;
 	UE_API virtual const TArray<FEdGraphPinReference>& GetLODPins() const override;
 	UE_API virtual UEdGraphPin* GetOutputPin() const override;
 	UE_API virtual void SetOutputPin(const UEdGraphPin* Pin) override;
@@ -54,9 +53,6 @@ public:
 	UPROPERTY(EditAnywhere, Category = ComponentMesh)
 	int32 NumLODs = 1;
 
-	UPROPERTY(EditAnywhere, Category = ComponentMesh, DisplayName = "Auto LOD Strategy")
-	ECustomizableObjectAutomaticLODStrategy AutoLODStrategy = ECustomizableObjectAutomaticLODStrategy::AutomaticFromMesh;
-
 	UPROPERTY()
 	TArray<FEdGraphPinReference> LODPins;
 
diff --git a/Engine/Plugins/Mutable/Source/CustomizableObjectEditor/Private/MuCOE/Nodes/CustomizableObjectNodeComponentMeshBase.h b/Engine/Plugins/Mutable/Source/CustomizableObjectEditor/Private/MuCOE/Nodes/CustomizableObjectNodeComponentMeshBase.h
--- a/Engine/Plugins/Mutable/Source/CustomizableObjectEditor/Private/MuCOE/Nodes/CustomizableObjectNodeComponentMeshBase.h
+++ b/Engine/Plugins/Mutable/Source/CustomizableObjectEditor/Private/MuCOE/Nodes/CustomizableObjectNodeComponentMeshBase.h
@@ -7,25 +7,6 @@
 #include "CustomizableObjectNodeComponentMeshBase.generated.h"
 
 
-UENUM()
-enum class ECustomizableObjectAutomaticLODStrategy : uint8
-{
-	// Use the same strategy than the parent component. If root, then use "Manual".
-	Inherited = 0 UMETA(DisplayName = "Inherit from parent component"),
-	// Don't try to generate LODs automatically for the child nodes. Only the ones tha explicitely define them will be used.
-	Manual = 1 UMETA(DisplayName = "Only manually created LODs"),
-	/** Automatic LODs from Mesh will try to replicate the explicit graph connections to higher LODs.
-	 *
-	 * For example, given a Mesh with n LODs connected to a Component m-th LOD, LODs m1, m2... will be automatically connected
-	 * using Mesh LODs n1, n2.
-	 *
-	 * Note: If the mesh section is disabled at LOD 'n' or doesn't exist, it will not be added from that LOD onwards.
-	 * Note: If the mesh section uses a different MaterialSlot at LOD 'n', it will not be added from that LOD onwards.
-	 */
-	AutomaticFromMesh = 2 UMETA(DisplayName = "Automatic from mesh")
-};
-
-
 UINTERFACE(MinimalAPI, Blueprintable)
 class UCustomizableObjectNodeComponentMeshInterface : public UInterface
 {
@@ -42,9 +23,6 @@ public:
 	/** Returns the num of Lods set in the Mesh Component Node */
 	virtual int32 GetNumLODs() = 0;
 
-	/** Returns the Lod Strategy set in the Mesh Component Node */
-	virtual ECustomizableObjectAutomaticLODStrategy GetAutoLODStrategy() = 0;
-
 	/** Returns a reference to an array with all the LOD pins of the Mesh Component Node */
 	virtual TArray<FEdGraphPinReference>& GetLODPins();
 
diff --git a/Engine/Plugins/Mutable/Source/CustomizableObjectEditor/Private/MuCOE/Nodes/CustomizableObjectNodeObject.cpp b/Engine/Plugins/Mutable/Source/CustomizableObjectEditor/Private/MuCOE/Nodes/CustomizableObjectNodeObject.cpp
--- a/Engine/Plugins/Mutable/Source/CustomizableObjectEditor/Private/MuCOE/Nodes/CustomizableObjectNodeObject.cpp
+++ b/Engine/Plugins/Mutable/Source/CustomizableObjectEditor/Private/MuCOE/Nodes/CustomizableObjectNodeObject.cpp
@@ -324,7 +324,6 @@ void UCustomizableObjectNodeObject::BackwardsCompatibleFixup(int32 CustomizableO
 
 					NodeComponentMeshAddTo->SetParentComponentName(MeshComponentName);
 					NodeComponentMeshAddTo->NumLODs = NumLODs_DEPRECATED;
-					NodeComponentMeshAddTo->AutoLODStrategy = AutoLODStrategy_DEPRECATED;
 					
 					// Create LOD pins.
 					NodeComponentMeshAddTo->FixupReconstructPins(CreateRemapPinsByName(), NodeComponentMeshBaseAllocateDefaultPins);
diff --git a/Engine/Plugins/Mutable/Source/CustomizableObjectEditor/Private/MuCOE/Nodes/CustomizableObjectNodeObject.h b/Engine/Plugins/Mutable/Source/CustomizableObjectEditor/Private/MuCOE/Nodes/CustomizableObjectNodeObject.h
--- a/Engine/Plugins/Mutable/Source/CustomizableObjectEditor/Private/MuCOE/Nodes/CustomizableObjectNodeObject.h
+++ b/Engine/Plugins/Mutable/Source/CustomizableObjectEditor/Private/MuCOE/Nodes/CustomizableObjectNodeObject.h
@@ -124,9 +124,6 @@ public:
 	UPROPERTY()
 	int32 NumLODs_DEPRECATED = 1;
 
-	UPROPERTY()
-	ECustomizableObjectAutomaticLODStrategy AutoLODStrategy_DEPRECATED = ECustomizableObjectAutomaticLODStrategy::AutomaticFromMesh;
-
 	UPROPERTY()
 	int32 NumMeshComponents_DEPRECATED = 1;
 
diff --git a/Engine/Plugins/Mutable/Source/MutableTools/Internal/MuT/NodeMeshSkeletalMeshBreak.h b/Engine/Plugins/Mutable/Source/MutableTools/Internal/MuT/NodeMeshSkeletalMeshBreak.h
--- a/Engine/Plugins/Mutable/Source/MutableTools/Internal/MuT/NodeMeshSkeletalMeshBreak.h
+++ b/Engine/Plugins/Mutable/Source/MutableTools/Internal/MuT/NodeMeshSkeletalMeshBreak.h
@@ -30,8 +30,6 @@ namespace UE::Mutable::Private
 		int8 BonePosePriority = 0;
 		int8 SocketPriority = 0;
 
-		bool bAutomaticLODs = false;
-
 		Ptr<NodeSkeletalMeshParameter> SkeletalMeshParameter;
 		
 		// Node Interface
diff --git a/Engine/Plugins/Mutable/Source/MutableTools/Private/MuT/CodeGenerator_Mesh.cpp b/Engine/Plugins/Mutable/Source/MutableTools/Private/MuT/CodeGenerator_Mesh.cpp
--- a/Engine/Plugins/Mutable/Source/MutableTools/Private/MuT/CodeGenerator_Mesh.cpp
+++ b/Engine/Plugins/Mutable/Source/MutableTools/Private/MuT/CodeGenerator_Mesh.cpp
@@ -1641,7 +1641,7 @@ namespace UE::Mutable::Private
 					
 				Op = new ASTOpMeshSkeletalMeshBreak();
 				Op->SkeletalMeshParameter = SkeletalMeshParameter;
-				Op->LODIndex = Node->LODIndex + (Node->bAutomaticLODs ? StaticOptions.LODIndex : 0);
+				Op->LODIndex = Node->LODIndex + StaticOptions.LODIndex;
 				Op->SectionIndex = Node->SectionIndex;
 				Op->ConversionFlags = Node->ConversionFlags;
 
