[Mutable] Fix Mutable mesh streaming not releasing the last generated resources until the mesh is destroyed.

#rnx
#Synced-CL 51492284

diff --git a/Engine/Plugins/Mutable/Source/CustomizableObject/Private/MuCO/CustomizableObjectMeshUpdate.cpp b/Engine/Plugins/Mutable/Source/CustomizableObject/Private/MuCO/CustomizableObjectMeshUpdate.cpp
--- a/Engine/Plugins/Mutable/Source/CustomizableObject/Private/MuCO/CustomizableObjectMeshUpdate.cpp
+++ b/Engine/Plugins/Mutable/Source/CustomizableObject/Private/MuCO/CustomizableObjectMeshUpdate.cpp
@@ -41,11 +41,15 @@ FCustomizableObjectMeshStreamIn::FCustomizableObjectMeshStreamIn(
 	EThreadType CreateResourcesThread) :
 	FSkeletalMeshStreamIn(InMesh, CreateResourcesThread)
 {
-	OperationData = InMesh->UpdateMeshContext;
-	check(OperationData)
+	check(InMesh->MeshContext);
+
+	UpdateContext = MakeShared<FMutableMeshUpdateContext>();
 	
-	OperationData->CurrentFirstLODIdx = CurrentFirstLODIdx;
-	OperationData->PendingFirstLODIdx = PendingFirstLODIdx;
+	UpdateContext->MeshContext = InMesh->MeshContext;
+	UpdateContext->Meshes.SetNum(InMesh->MeshContext->MeshIDs.Num());
+
+	UpdateContext->CurrentFirstLODIdx = CurrentFirstLODIdx;
+	UpdateContext->PendingFirstLODIdx = PendingFirstLODIdx;
 
 	PushTask(FContext(InMesh, TT_None), TT_Async, SRA_UPDATE_CALLBACK(DoInitiate), TT_None, nullptr);
 }
@@ -143,14 +147,14 @@ void FCustomizableObjectMeshStreamIn::DoCreateBuffers(const FContext& Context)
 
 namespace impl
 {
-	void Task_Mutable_UpdateMesh_End(const TSharedPtr<FMutableUpdateMeshContext>& OperationData, TRefCountPtr<FCustomizableObjectMeshStreamIn>& StreamingTask)
+	void Task_Mutable_UpdateMesh_End(const TSharedPtr<FMutableMeshUpdateContext>& UpdateContext, TRefCountPtr<FCustomizableObjectMeshStreamIn>& StreamingTask)
 	{
 		MUTABLE_CPUPROFILER_SCOPE(Task_Mutable_UpdateMesh_End);
 		
 		// End update
 		const bool bClearRoms = UCustomizableObjectSystem::ShouldClearWorkingMemoryOnUpdateEnd();
 		const bool bFreeCache = !CVarReuseLiveInstanceCacheBetweenUpdates.GetValueOnAnyThread();
-		OperationData->System->EndUpdate(OperationData->LiveInstance, bClearRoms, bFreeCache);
+		UpdateContext->MeshContext->System->EndUpdate(UpdateContext->MeshContext->LiveInstance, bClearRoms, bFreeCache);
 
 		StreamingTask->OnUpdateMeshFinished();
 		
@@ -159,35 +163,35 @@ namespace impl
 
 	
 	void Task_Mutable_UpdateMesh_Loop(
-		const TSharedPtr<FMutableUpdateMeshContext>& OperationData,
+		const TSharedPtr<FMutableMeshUpdateContext>& UpdateContext,
 		TRefCountPtr<FCustomizableObjectMeshStreamIn>& Task,
 		int32 LODIndex)
 	{
 		MUTABLE_CPUPROFILER_SCOPE(Task_Mutable_UpdateMesh_Loop);
 
-		if (Task->IsCancelled() || LODIndex == OperationData->CurrentFirstLODIdx + OperationData->AssetLODBias)
+		if (Task->IsCancelled() || LODIndex == UpdateContext->CurrentFirstLODIdx + UpdateContext->AssetLODBias)
 		{
-			Task_Mutable_UpdateMesh_End(OperationData, Task);
+			Task_Mutable_UpdateMesh_End(UpdateContext, Task);
 			return;
 		}
 
-		const UE::Mutable::Private::FMeshId& MeshID = OperationData->MeshIDs[LODIndex];
+		const UE::Mutable::Private::FMeshId& MeshID = UpdateContext->MeshContext->MeshIDs[LODIndex];
 
 		constexpr UE::Mutable::Private::EMeshContentFlags MeshContentFilter = UE::Mutable::Private::EMeshContentFlags::AllFlags;
-		UE::Tasks::TTask<UE::Mutable::Private::TManagedPtr<const UE::Mutable::Private::FMesh>> GetMeshTask = OperationData->System->GetMesh({}, OperationData->LiveInstance, MeshID, MeshContentFilter);
+		UE::Tasks::TTask<UE::Mutable::Private::TManagedPtr<const UE::Mutable::Private::FMesh>> GetMeshTask = UpdateContext->MeshContext->System->GetMesh({}, UpdateContext->MeshContext->LiveInstance, MeshID, MeshContentFilter);
 
 		UE::Tasks::AddNested(UE::Tasks::Launch(TEXT("Task_MutableGetMeshes_GetMesh_Post"), [=]() mutable
 			{
-				OperationData->Meshes[LODIndex] = GetMeshTask.GetResult();
+				UpdateContext->Meshes[LODIndex] = GetMeshTask.GetResult();
 
-				Task_Mutable_UpdateMesh_Loop(OperationData, Task, LODIndex + 1);
+				Task_Mutable_UpdateMesh_Loop(UpdateContext, Task, LODIndex + 1);
 			},
 			GetMeshTask,
 			LowLevelTasks::ETaskPriority::Inherit));
 	}
 
 	
-	void Task_Mutable_UpdateMesh(const TSharedPtr<FMutableUpdateMeshContext> OperationData, TRefCountPtr<FCustomizableObjectMeshStreamIn>& Task)
+	void Task_Mutable_UpdateMesh(const TSharedPtr<FMutableMeshUpdateContext> UpdateContext, TRefCountPtr<FCustomizableObjectMeshStreamIn>& Task)
 	{
 		MUTABLE_CPUPROFILER_SCOPE(Task_Mutable_UpdateMesh);
 
@@ -198,8 +202,8 @@ namespace impl
 
 		TRACE_BEGIN_REGION(UE_MUTABLE_UPDATE_MESH_REGION);
 
-		TSharedPtr<UE::Mutable::Private::FSystem> System = OperationData->System;
-		TSharedRef<UE::Mutable::Private::FLiveInstance> LiveInstance = OperationData->LiveInstance;
+		TSharedPtr<UE::Mutable::Private::FSystem> System = UpdateContext->MeshContext->System;
+		TSharedRef<UE::Mutable::Private::FLiveInstance> LiveInstance = UpdateContext->MeshContext->LiveInstance;
 		
 #if WITH_EDITOR
 		const TSharedPtr<const UE::Mutable::Private::FModel, ESPMode::ThreadSafe> Model = LiveInstance->Model;
@@ -227,7 +231,7 @@ namespace impl
 		
 		check(Instance);
 
-		Task_Mutable_UpdateMesh_Loop(OperationData, Task, OperationData->PendingFirstLODIdx + OperationData->AssetLODBias);
+		Task_Mutable_UpdateMesh_Loop(UpdateContext, Task, UpdateContext->PendingFirstLODIdx + UpdateContext->AssetLODBias);
 	}
 	
 } // namespace
@@ -251,11 +255,11 @@ void FCustomizableObjectMeshStreamIn::RequestMeshUpdate(const FContext& Context)
 		return;
 	}
 
-	check(OperationData.IsValid());
+	check(UpdateContext.IsValid());
 
 #if WITH_EDITOR
 	// Recompiling a CO in the editor will invalidate the previously generated Model. Check that it is valid before accessing the streamed data.
-	const TSharedPtr<const UE::Mutable::Private::FModel> Model = OperationData->LiveInstance->Model;
+	const TSharedPtr<const UE::Mutable::Private::FModel> Model = UpdateContext->MeshContext->LiveInstance->Model;
 	if (!Model || !Model->IsValid())
 	{
 		Abort();
@@ -263,7 +267,7 @@ void FCustomizableObjectMeshStreamIn::RequestMeshUpdate(const FContext& Context)
 	}
 #endif
 
-	OperationData->AssetLODBias = Context.AssetLODBias;
+	UpdateContext->AssetLODBias = Context.AssetLODBias;
 
 	UCustomizableObjectSystemPrivate* CustomizableObjectSystem = UCustomizableObjectSystem::GetInstance()->GetPrivate();
 	
@@ -273,18 +277,18 @@ void FCustomizableObjectMeshStreamIn::RequestMeshUpdate(const FContext& Context)
 	{
 		CustomizableObjectSystem->MutableTaskGraph.AddMeshStreamingTask(
 			TEXT("Mutable_MeshUpdate"),
-			[SharedOperationData = OperationData, RefThis = TRefCountPtr<FCustomizableObjectMeshStreamIn>(this)]() mutable
+			[SharedUpdateContext = UpdateContext, RefThis = TRefCountPtr<FCustomizableObjectMeshStreamIn>(this)]() mutable
 			{
-				impl::Task_Mutable_UpdateMesh(SharedOperationData, RefThis);
+				impl::Task_Mutable_UpdateMesh(SharedUpdateContext, RefThis);
 			});
 	}
 	else
 	{
 		MutableTaskId = CustomizableObjectSystem->MutableTaskGraph.AddMutableThreadTaskLowPriority(
 			TEXT("Mutable_MeshUpdate"),
-			[SharedOperationData = OperationData, RefThis = TRefCountPtr<FCustomizableObjectMeshStreamIn>(this)]() mutable
+			[SharedUpdateContext = UpdateContext, RefThis = TRefCountPtr<FCustomizableObjectMeshStreamIn>(this)]() mutable
 			{
-				impl::Task_Mutable_UpdateMesh(SharedOperationData, RefThis);
+				impl::Task_Mutable_UpdateMesh(SharedUpdateContext, RefThis);
 			});
 	}
 
@@ -313,7 +317,7 @@ void FCustomizableObjectMeshStreamIn::ConvertMesh(const FContext& Context, bool&
 	{
 		const int32 LODIndexAbsolute = LODIndex + Context.AssetLODBias;
 
-		UE::Mutable::Private::TManagedPtr<const UE::Mutable::Private::FMesh> MutableMesh = OperationData->Meshes[LODIndexAbsolute];
+		UE::Mutable::Private::TManagedPtr<const UE::Mutable::Private::FMesh> MutableMesh = UpdateContext->Meshes[LODIndexAbsolute];
 
 		if (!MutableMesh)
 		{
@@ -334,8 +338,8 @@ void FCustomizableObjectMeshStreamIn::ConvertMesh(const FContext& Context, bool&
 
 		FSkeletalMeshLODRenderData& LODResource = *Context.LODResourcesView[LODIndex];
 		UnrealConversionUtils::CopyMutableVertexBuffers(LODResource, MutableMesh.Get(), bNeedsCPUAccess);
-		UnrealConversionUtils::CopyMutableIndexBuffers(LODResource, MutableMesh.Get(), Mesh->SurfaceIDs[LODIndexAbsolute], bOutMarkRenderStateDirty);
-		UnrealConversionUtils::CopyMutableSkinWeightProfilesBuffers(LODResource, *const_cast<UCustomizableObjectSkeletalMesh*>(Mesh), LODIndexAbsolute, MutableMesh.Get(), Mesh->SkinWeightProfileIDs);
+		UnrealConversionUtils::CopyMutableIndexBuffers(LODResource, MutableMesh.Get(), UpdateContext->MeshContext->SurfaceIDs[LODIndexAbsolute], bOutMarkRenderStateDirty);
+		UnrealConversionUtils::CopyMutableSkinWeightProfilesBuffers(LODResource, *const_cast<UCustomizableObjectSkeletalMesh*>(Mesh), LODIndexAbsolute, MutableMesh.Get(), UpdateContext->MeshContext->SkinWeightProfileIDs);
 		UnrealConversionUtils::MorphTargetVertexInfoBuffers(LODResource, *Mesh, *MutableMesh.Get(), LODIndexAbsolute, bGenerateCPUMorphTargetsIfNeeded);
 		UnrealConversionUtils::ClothVertexBuffers(LODResource, *MutableMesh.Get());
 		
@@ -343,7 +347,7 @@ void FCustomizableObjectMeshStreamIn::ConvertMesh(const FContext& Context, bool&
 	}
 
 	// Clear MeshUpdate data
-	OperationData = nullptr;
+	UpdateContext = nullptr;
 }
 
 
diff --git a/Engine/Plugins/Mutable/Source/CustomizableObject/Private/MuCO/CustomizableObjectMeshUpdate.h b/Engine/Plugins/Mutable/Source/CustomizableObject/Private/MuCO/CustomizableObjectMeshUpdate.h
--- a/Engine/Plugins/Mutable/Source/CustomizableObject/Private/MuCO/CustomizableObjectMeshUpdate.h
+++ b/Engine/Plugins/Mutable/Source/CustomizableObject/Private/MuCO/CustomizableObjectMeshUpdate.h
@@ -15,6 +15,7 @@ CustomizableObjectMeshUpdate.h: Helpers to stream in skeletal mesh LODs.
 #include "MuR/Mesh.h"
 
 struct FModelStreamableBulkData;
+struct FMutableMeshContext;
 class UCustomizableObjectSkeletalMesh;
 
 namespace UE::Mutable::Private 
@@ -22,27 +23,18 @@ namespace UE::Mutable::Private
 	class FParameters;
 }
 
-struct FMutableUpdateMeshContext
+
+/** Runtime data used during a mutable mesh LOD update */
+struct FMutableMeshUpdateContext
 {
-	FMutableUpdateMeshContext(const TSharedRef<UE::Mutable::Private::FLiveInstance>& InLiveInstance)
-		: LiveInstance(InLiveInstance)
-	{
-		
-	}
-	
-	TSharedPtr<UE::Mutable::Private::FSystem> System;
-	TSharedRef<UE::Mutable::Private::FLiveInstance> LiveInstance;
-
-	TSharedPtr<FModelStreamableBulkData> ModelStreamableBulkData;
-	
-	// IDs of the meshes to generate per LOD.
-	TArray<UE::Mutable::Private::FMeshId> MeshIDs;
+	/** Mesh update context */
+	TSharedPtr<FMutableMeshContext> MeshContext;
 
 	// Meshes generated per LOD
 	TArray<UE::Mutable::Private::TManagedPtr<const UE::Mutable::Private::FMesh>> Meshes;
 
 	/** SkeletalMesh */
-	/** The resident first LOD resource index.With domain = [0, ResourceState.NumLODs[.NOT THE ASSET LOD INDEX! */ 
+	/** The resident first LOD resource index.With domain = [0, ResourceState.NumLODs[.NOT THE ASSET LOD INDEX! */
 	int32 CurrentFirstLODIdx = INDEX_NONE;
 
 	/** The requested first LOD resource index. With domain = [0, ResourceState.NumLODs[. NOT THE ASSET LOD INDEX! */
@@ -86,8 +78,8 @@ private:
 
 	void MarkRenderStateDirty(const FContext& Context);
 
-	/** Context of the MeshUpdate. */
-	TSharedPtr<FMutableUpdateMeshContext> OperationData;
+	/** Data of the MeshUpdate. */
+	TSharedPtr<FMutableMeshUpdateContext> UpdateContext;
 
 	/** MeshUpdate task Id to cancel the task if the stream in task is aborted. */
 	uint32 MutableTaskId = 0;
diff --git a/Engine/Plugins/Mutable/Source/CustomizableObject/Private/MuCO/CustomizableObjectSkeletalMesh.cpp b/Engine/Plugins/Mutable/Source/CustomizableObject/Private/MuCO/CustomizableObjectSkeletalMesh.cpp
--- a/Engine/Plugins/Mutable/Source/CustomizableObject/Private/MuCO/CustomizableObjectSkeletalMesh.cpp
+++ b/Engine/Plugins/Mutable/Source/CustomizableObject/Private/MuCO/CustomizableObjectSkeletalMesh.cpp
@@ -14,50 +14,40 @@
 #include UE_INLINE_GENERATED_CPP_BY_NAME(CustomizableObjectSkeletalMesh)
 
 
-
-
 void UCustomizableObjectSkeletalMesh::InitMutableStreamingData(const TSharedRef<FUpdateContextPrivate>& InOperationData, const FName& ComponentName, const int32 InstanceUpdateFirstLOD, const int32 LODCount)
 {
 	UCustomizableObject& CustomizableObject = *InOperationData->Object;
-
-	ModelStreamableBulkData = CustomizableObject.GetPrivate()->GetModelStreamableBulkData();
 	
 	// Debug info
 	CustomizableObjectPathName = GetNameSafe(&CustomizableObject);
 
-	MeshIDs.Init({}, MAX_MESH_LOD_COUNT);
-	SurfaceIDs.SetNum(MAX_MESH_LOD_COUNT);
+	// Init properties
+	check(InOperationData->LiveInstance.IsValid());
+	MeshContext = MakeShared<FMutableMeshContext>(InOperationData->LiveInstance.ToSharedRef());
+
+	// This must run in the mutable thread.
+	check(UCustomizableObjectSystem::IsCreated());
+	MeshContext->System = UCustomizableObjectSystem::GetInstance()->GetPrivate()->MutableSystem;
+	MeshContext->ModelStreamableBulkData = CustomizableObject.GetPrivate()->GetModelStreamableBulkData();
+	
+	MeshContext->MeshIDs.Init({}, LODCount);
+	MeshContext->SurfaceIDs.SetNum(LODCount);
 
 	for (int32 LODIndex = InOperationData->FirstLODAvailable[ComponentName]; LODIndex < LODCount; ++LODIndex)
 	{
 		const TSharedRef<const FInstanceUpdateData::FLOD>& LOD = InOperationData->InstanceUpdateData.LODs[InstanceUpdateFirstLOD + LODIndex];
-		MeshIDs[LODIndex] = LOD->MeshId;
+		MeshContext->MeshIDs[LODIndex] = LOD->MeshId;
 
-		SurfaceIDs[LODIndex].SetNum(LOD->SurfaceCount);
+		MeshContext->SurfaceIDs[LODIndex].SetNum(LOD->SurfaceCount);
 		check(LOD->Mesh && LOD->Mesh->GetSurfaceCount() == LOD->SurfaceCount);
 
 		for (int32 SurfaceIndex = 0; SurfaceIndex < LOD->SurfaceCount; ++SurfaceIndex)
 		{
-			SurfaceIDs[LODIndex][SurfaceIndex] = LOD->Mesh->GetSurfaceId(SurfaceIndex);
+			MeshContext->SurfaceIDs[LODIndex][SurfaceIndex] = LOD->Mesh->GetSurfaceId(SurfaceIndex);
 		}
 	}
 
 	const UModelResources& ModelResources = CustomizableObject.GetPrivate()->GetModelResourcesChecked();
-
-	{
-		// Init properties
-		check(InOperationData->LiveInstance.IsValid());
-		UpdateMeshContext = MakeShared<FMutableUpdateMeshContext>(InOperationData->LiveInstance.ToSharedRef());
-
-		// This must run in the mutable thread.
-		check(UCustomizableObjectSystem::IsCreated());
-		UpdateMeshContext->System = UCustomizableObjectSystem::GetInstance()->GetPrivate()->MutableSystem;
-		UpdateMeshContext->ModelStreamableBulkData = ModelStreamableBulkData;
-		
-		UpdateMeshContext->MeshIDs = MeshIDs;
-		UpdateMeshContext->Meshes.SetNum(MeshIDs.Num());
-	}
-	
 	for (const FSkinWeightProfileInfo& SkinWeightProfile : GetSkinWeightProfiles())
 	{
 		const FMutableSkinWeightProfileInfo* ProfileInfo = ModelResources.SkinWeightProfilesInfo.FindByPredicate(
@@ -65,7 +55,7 @@ void UCustomizableObjectSkeletalMesh::InitMutableStreamingData(const TSharedRef<
 
 		if (ensure(ProfileInfo))
 		{
-			SkinWeightProfileIDs.Add({ ProfileInfo->NameId, ProfileInfo->Name });
+			MeshContext->SkinWeightProfileIDs.Add({ ProfileInfo->NameId, ProfileInfo->Name });
 		}
 	}
 }
diff --git a/Engine/Plugins/Mutable/Source/CustomizableObject/Private/MuCO/UnrealConversionUtils.cpp b/Engine/Plugins/Mutable/Source/CustomizableObject/Private/MuCO/UnrealConversionUtils.cpp
--- a/Engine/Plugins/Mutable/Source/CustomizableObject/Private/MuCO/UnrealConversionUtils.cpp
+++ b/Engine/Plugins/Mutable/Source/CustomizableObject/Private/MuCO/UnrealConversionUtils.cpp
@@ -1506,8 +1506,6 @@ namespace UnrealConversionUtils
 									FMemoryReaderView Ar(TArrayView<uint8>(BulkData, IORequest->GetSize()), true);
 
 									constexpr uint8 DummyStripFlags = 0;
-
-									// NeedsCPUaccess must be false. Otherwise, the FSkeletalMeshLODRenderData buffers wouldn't free their resources, leading to a memory leak.
 									const bool bForceKeepCPUResources = false;
 									const bool bNeedsCPUAccess = false;
 									Context->StreamedLOD.SerializeStreamedData(Ar, const_cast<USkeletalMesh*>(SkeletalMesh.Get()), LODIndex, DummyStripFlags, bNeedsCPUAccess, bForceKeepCPUResources);
diff --git a/Engine/Plugins/Mutable/Source/CustomizableObject/Public/MuCO/CustomizableObjectSkeletalMesh.h b/Engine/Plugins/Mutable/Source/CustomizableObject/Public/MuCO/CustomizableObjectSkeletalMesh.h
--- a/Engine/Plugins/Mutable/Source/CustomizableObject/Public/MuCO/CustomizableObjectSkeletalMesh.h
+++ b/Engine/Plugins/Mutable/Source/CustomizableObject/Public/MuCO/CustomizableObjectSkeletalMesh.h
@@ -20,9 +20,29 @@ class UCustomizableInstancePrivate;
 class FUpdateContextPrivate;
 class FCustomizableObjectMeshStreamIn;
 struct FModelStreamableBulkData;
-struct FMutableUpdateMeshContext;
 class FUnrealMutableResourceProvider;
 
+/**
+ * Context required to stream in/out mesh LODs of a UCustomizableObjectSkeletalMesh
+ */
+struct FMutableMeshContext
+{
+	FMutableMeshContext(const TSharedRef<UE::Mutable::Private::FLiveInstance>& InLiveInstance)
+		: LiveInstance(InLiveInstance) {}
+
+	TSharedPtr<UE::Mutable::Private::FSystem> System;
+	TSharedRef<UE::Mutable::Private::FLiveInstance> LiveInstance;
+
+	TSharedPtr<FModelStreamableBulkData> ModelStreamableBulkData;
+
+	// IDs of the meshes to generate per LOD.
+	TArray<UE::Mutable::Private::FMeshId> MeshIDs;
+
+	/** IDs per LOD */
+	TArray<TArray<uint32>> SurfaceIDs;
+	TArray<TPair<uint32, FName>> SkinWeightProfileIDs;
+};
+
 
 /**
  * CustomizableObjectSkeletalMesh is a class inheriting from USkeletalMesh with the sole purpose of allowing skeletal meshes generated by Mutable
@@ -45,15 +65,7 @@ class UCustomizableObjectSkeletalMesh : public USkeletalMesh
 
 	FString CustomizableObjectPathName;
 
-	TSharedPtr<FModelStreamableBulkData> ModelStreamableBulkData;
-
-	TSharedPtr<FMutableUpdateMeshContext> UpdateMeshContext;
-
-
-	/** Mesh IDs per LOD */
-	TArray<UE::Mutable::Private::FMeshId> MeshIDs;
-	TArray<TArray<uint32>> SurfaceIDs;
-	TArray<TPair<uint32, FName>> SkinWeightProfileIDs;
+	TSharedPtr<FMutableMeshContext> MeshContext;
 
 	FMeshCache::FId MeshCacheId;
 };
