[Mutable] Enable ImageMulticompose optimization

#rnx
#Synced-CL 51492284

diff --git a/Engine/Plugins/Mutable/Source/MutableTools/Private/MuT/CodeGenerator.cpp b/Engine/Plugins/Mutable/Source/MutableTools/Private/MuT/CodeGenerator.cpp
--- a/Engine/Plugins/Mutable/Source/MutableTools/Private/MuT/CodeGenerator.cpp
+++ b/Engine/Plugins/Mutable/Source/MutableTools/Private/MuT/CodeGenerator.cpp
@@ -82,7 +82,7 @@ namespace UE::Mutable::Private
 {
 	TAutoConsoleVariable<bool> CVarAllowMulticompose(
 		TEXT("mutable.AllowImageMulticompose"),
-		false,
+		true,
 		TEXT("If set to true, enables the use of multicompose."),
 		ECVF_Default);
 
