[Mutable] Fix Customizable Object being GC during package splitter.

#Synced-CL 51537794

diff --git a/Engine/Plugins/Mutable/Source/CustomizableObjectEditor/Private/MuCOE/CustomizableObjectCookPackageSplitter.cpp b/Engine/Plugins/Mutable/Source/CustomizableObjectEditor/Private/MuCOE/CustomizableObjectCookPackageSplitter.cpp
--- a/Engine/Plugins/Mutable/Source/CustomizableObjectEditor/Private/MuCOE/CustomizableObjectCookPackageSplitter.cpp
+++ b/Engine/Plugins/Mutable/Source/CustomizableObjectEditor/Private/MuCOE/CustomizableObjectCookPackageSplitter.cpp
@@ -75,7 +75,8 @@ FString GenerateRelativePath(const UCustomizableObject& Object, UE::Mutable::Pri
 ICookPackageSplitter::FGenerationManifest FCustomizableObjectCookPackageSplitter::ReportGenerationManifest(const UPackage* OwnerPackage, const UObject* OwnerObject)
 {
 	const UCustomizableObject* Object = CastChecked<UCustomizableObject>(OwnerObject);
-
+	StrongObject = TStrongObjectPtr(Object);
+	
 	// All platforms should have the same resources
 
 	const UModelResources* ModelResources = FindModelResources(*Object);
diff --git a/Engine/Plugins/Mutable/Source/CustomizableObjectEditor/Private/MuCOE/CustomizableObjectCookPackageSplitter.h b/Engine/Plugins/Mutable/Source/CustomizableObjectEditor/Private/MuCOE/CustomizableObjectCookPackageSplitter.h
--- a/Engine/Plugins/Mutable/Source/CustomizableObjectEditor/Private/MuCOE/CustomizableObjectCookPackageSplitter.h
+++ b/Engine/Plugins/Mutable/Source/CustomizableObjectEditor/Private/MuCOE/CustomizableObjectCookPackageSplitter.h
@@ -20,4 +20,6 @@ public:
 	virtual EGeneratedRequiresGenerator DoesGeneratedRequireGenerator() override;
 	virtual bool UseInternalReferenceToAvoidGarbageCollect() override;
 	virtual bool RequiresEmptyPackageBeforePopulate(const UPackage* Package) override;
+	
+	TStrongObjectPtr<const UCustomizableObject> StrongObject;
 };
