[Mutable] Fix External Operations with Mesh inputs not propagating generated data upwards.

#Synced-CL 51537794

diff --git a/Engine/Plugins/Mutable/Source/MutableTools/Private/MuT/ASTOpMeshMaskDiff.cpp b/Engine/Plugins/Mutable/Source/MutableTools/Private/MuT/ASTOpMeshMaskDiff.cpp
--- a/Engine/Plugins/Mutable/Source/MutableTools/Private/MuT/ASTOpMeshMaskDiff.cpp
+++ b/Engine/Plugins/Mutable/Source/MutableTools/Private/MuT/ASTOpMeshMaskDiff.cpp
@@ -159,7 +159,7 @@ namespace UE::Mutable::Private
 		}
 
 		// If we didn't optimize the fragment, see ifg we can optimize the source.
-		if (!ResultOp)
+		if (!ResultOp && Source)
 		{
 			EOpType SourceType = Source->GetOpType();
 			switch (SourceType)
diff --git a/Engine/Plugins/Mutable/Source/MutableTools/Private/MuT/CodeGenerator.h b/Engine/Plugins/Mutable/Source/MutableTools/Private/MuT/CodeGenerator.h
--- a/Engine/Plugins/Mutable/Source/MutableTools/Private/MuT/CodeGenerator.h
+++ b/Engine/Plugins/Mutable/Source/MutableTools/Private/MuT/CodeGenerator.h
@@ -996,6 +996,9 @@ namespace UE::Mutable::Private
     	struct FExternalGenerationResult
     	{
     		Ptr<ASTOp> Op;
+    		
+    		// If set, at least one of the external operation inputs is a mesh
+    		TOptional<FMeshGenerationResult> MeshGenerationResult;
     	};
     	
     	typedef TMap<FGeneratedExternalCacheKey, FExternalGenerationResult> FGeneratedExternalMap;
diff --git a/Engine/Plugins/Mutable/Source/MutableTools/Private/MuT/CodeGenerator_External.cpp b/Engine/Plugins/Mutable/Source/MutableTools/Private/MuT/CodeGenerator_External.cpp
--- a/Engine/Plugins/Mutable/Source/MutableTools/Private/MuT/CodeGenerator_External.cpp
+++ b/Engine/Plugins/Mutable/Source/MutableTools/Private/MuT/CodeGenerator_External.cpp
@@ -106,9 +106,6 @@ namespace UE::Mutable::Private
 			External->Type = EOpType::IS_EXTERNAL;
 		}
 		
-		FMeshGenerationResult Result;
-		Result.MeshOp = External;
-
 		const int32 NumInputs = ExternalNode->Inputs.Num();
     			
 		External->Inputs.Reserve(NumInputs);
@@ -128,9 +125,28 @@ namespace UE::Mutable::Private
 					UE::Tasks::MakeCompletedTask<FMeshGenerationDynamicOptions>();
     					
 				FMeshTask Task = GenerateMesh(MeshStaticOptions, MeshDynamicOptions, static_pointer_cast<NodeMesh>(Input));
-				Task.Wait();
+				
+				if (WaitCallback.IsSet())
+				{
+					while (!Task.IsCompleted())
+					{
+						WaitCallback();
+					}
+				}
+				
+				FMeshGenerationResult MeshResult = Task.GetResult();
 
-				External->Inputs.Add(ASTChild(External.get() ,Task.GetResult().MeshOp));
+				External->Inputs.Add(ASTChild(External.get(), MeshResult.MeshOp));
+				
+				if (!OutResult.MeshGenerationResult.IsSet())
+				{
+					OutResult.MeshGenerationResult = FMeshGenerationResult();
+				}
+				
+				// Merge all generation results from all inputs.
+				OutResult.MeshGenerationResult->GeneratedLayouts.Append(MeshResult.GeneratedLayouts);
+				OutResult.MeshGenerationResult->LayoutOps.Append(MeshResult.LayoutOps);
+				OutResult.MeshGenerationResult->ExtraMeshLayouts.Append(MeshResult.ExtraMeshLayouts);
 			}
 			else if (Input->GetType()->IsA(NodeImage::GetStaticType()))
 			{
@@ -189,7 +205,13 @@ namespace UE::Mutable::Private
 				unimplemented(); // GenerateMutableSourceExtension has compiled an input type that is not supported.
 			}
 		}
-
+		
+		if (OutResult.MeshGenerationResult.IsSet())
+		{
+			OutResult.MeshGenerationResult->MeshOp = External;	
+			OutResult.MeshGenerationResult->BaseMeshOp = External;			
+		}
+		
 		OutResult.Op = External;
 	}
 
diff --git a/Engine/Plugins/Mutable/Source/MutableTools/Private/MuT/CodeGenerator_Image.cpp b/Engine/Plugins/Mutable/Source/MutableTools/Private/MuT/CodeGenerator_Image.cpp
--- a/Engine/Plugins/Mutable/Source/MutableTools/Private/MuT/CodeGenerator_Image.cpp
+++ b/Engine/Plugins/Mutable/Source/MutableTools/Private/MuT/CodeGenerator_Image.cpp
@@ -2011,6 +2011,8 @@ namespace UE::Mutable::Private
 		GenerateExternal(ExtensionResult, Options, InNode->Node);
 
 		Result.op = ExtensionResult.Op;
+		Result.Parameter = false;
+		Result.bCanMulticompose = false;
 	}
 
 	
diff --git a/Engine/Plugins/Mutable/Source/MutableTools/Private/MuT/CodeGenerator_Mesh.cpp b/Engine/Plugins/Mutable/Source/MutableTools/Private/MuT/CodeGenerator_Mesh.cpp
--- a/Engine/Plugins/Mutable/Source/MutableTools/Private/MuT/CodeGenerator_Mesh.cpp
+++ b/Engine/Plugins/Mutable/Source/MutableTools/Private/MuT/CodeGenerator_Mesh.cpp
@@ -1715,6 +1715,12 @@ namespace UE::Mutable::Private
 				GenerateExternal(ExtensionResult, Options, Node->Node);
 
 				FMeshGenerationResult MeshResult;
+
+				if (ExtensionResult.MeshGenerationResult.IsSet())
+				{
+					MeshResult = ExtensionResult.MeshGenerationResult.GetValue();
+				}
+					
 				MeshResult.MeshOp = ExtensionResult.Op;
 					
 				return MeshResult;
