[Mutable] Fix realtime morph not generating when mesh streaming is enabled.

#rnx

#Synced-CL 51538701

diff --git a/Engine/Plugins/Mutable/Source/CustomizableObject/Private/MuCO/CustomizableObjectSystem.cpp b/Engine/Plugins/Mutable/Source/CustomizableObject/Private/MuCO/CustomizableObjectSystem.cpp
--- a/Engine/Plugins/Mutable/Source/CustomizableObject/Private/MuCO/CustomizableObjectSystem.cpp
+++ b/Engine/Plugins/Mutable/Source/CustomizableObject/Private/MuCO/CustomizableObjectSystem.cpp
@@ -1681,6 +1681,9 @@ namespace impl
 					continue;
 				}
 
+				// At this point, if the mesh HasMorph then we must have something in morph data.
+				check(LOD->Mesh->MorphDataBuffer.Num());
+
 				TArray<FMorphTargetLODModel> MorphTargets;
 				ReconstructMorphTargetsFromMeshCompressedData(*LOD->Mesh, MorphTargets, UE::Mutable::Private::EMorphUsageFlags::RealTime);
 
@@ -3649,6 +3652,32 @@ void UCustomizableObjectSystemPrivate::LogUpdateMessage(const FString& Message,
 	}
 }
 
+int32 CalculateMorphNumVertices(TConstArrayView<uint32> MorphHeadersDataView)
+{
+	using namespace UE::MorphTargetVertexCodec;
+
+	if (MorphHeadersDataView.Num() == 0)
+	{
+		return 0;
+	}
+
+	int32 NumBatches = MorphHeadersDataView.Num() / UE::MorphTargetVertexCodec::NumBatchHeaderDwords;
+
+	int32 NumMorphVertices = 0;
+	for (int32 BatchHeaderIndex = 0; BatchHeaderIndex < NumBatches; ++BatchHeaderIndex)
+	{
+		TConstArrayView<uint32> BatchHeaderData(
+			MorphHeadersDataView.GetData() + BatchHeaderIndex * UE::MorphTargetVertexCodec::NumBatchHeaderDwords,
+			UE::MorphTargetVertexCodec::NumBatchHeaderDwords);
+
+		FDeltaBatchHeader BatchHeader;
+		ReadHeader(BatchHeader, BatchHeaderData);
+
+		NumMorphVertices += BatchHeader.NumElements;
+	}
+
+	return NumMorphVertices;
+}
 
 #if !WITH_EDITOR
 void ReconstructMorphTargetsFromMeshCompressedData(const UE::Mutable::Private::FMesh& Mesh, TArray<FMorphTargetCompressedLODModel>& OutMorphTargets, UE::Mutable::Private::EMorphUsageFlags UsageFilter)
@@ -3736,18 +3765,12 @@ void ReconstructMorphTargetsFromMeshCompressedData(const UE::Mutable::Private::F
 			const int32 BatchStartOffset = Mesh.Morph.BatchStartOffsetPerMorph[MorphIndex];
 			const int32 MorphNumBatches  = Mesh.Morph.BatchesPerMorph[MorphIndex];
 
-			int32 NumMorphVertices = 0;
-			for (int32 BatchHeaderIndex = 0; BatchHeaderIndex < MorphNumBatches; ++BatchHeaderIndex)
-			{
-				TConstArrayView<uint32> BatchHeaderData(
-						&Mesh.MorphDataBuffer[(BatchStartOffset + BatchHeaderIndex) * NumBatchHeaderDwords], 
-						NumBatchHeaderDwords);
+			check((BatchStartOffset + MorphNumBatches)*NumBatchHeaderDwords <= Mesh.MorphDataBuffer.Num());
+			TConstArrayView<uint32> MorphHeadersView = TConstArrayView<uint32>(
+				Mesh.MorphDataBuffer.GetData() + BatchStartOffset * NumBatchHeaderDwords,
+				MorphNumBatches * NumBatchHeaderDwords);
 
-				FDeltaBatchHeader BatchHeader;
-				ReadHeader(BatchHeader, BatchHeaderData);
-
-				NumMorphVertices += BatchHeader.NumElements;
-			}
+			int32 NumMorphVertices = CalculateMorphNumVertices(MorphHeadersView);
 			
 			OutMorphTargets[MorphIndex].Vertices.SetNumUninitialized(NumMorphVertices);
 			OutMorphTargets[MorphIndex].NumVertices = NumMorphVertices;
diff --git a/Engine/Plugins/Mutable/Source/CustomizableObject/Private/MuCO/UnrealConversionUtils.cpp b/Engine/Plugins/Mutable/Source/CustomizableObject/Private/MuCO/UnrealConversionUtils.cpp
--- a/Engine/Plugins/Mutable/Source/CustomizableObject/Private/MuCO/UnrealConversionUtils.cpp
+++ b/Engine/Plugins/Mutable/Source/CustomizableObject/Private/MuCO/UnrealConversionUtils.cpp
@@ -994,6 +994,12 @@ namespace UnrealConversionUtils
 				CompressedMorphsToInit[*FoundIndex] = &CompressedMorphs[MorphIndex];
 				MorphTargetMinimumValues[*FoundIndex] = MutableMesh.Morph.MinimumValuePerMorph[MorphIndex]; 
 				MorphTargetMaximumValues[*FoundIndex] = MutableMesh.Morph.MaximumValuePerMorph[MorphIndex];
+
+				check(Owner.GetMorphTargets().IsValidIndex(*FoundIndex));
+				check(MorphLOD->GetMorphLODModels().IsValidIndex(LODIndex));
+
+				UMorphTarget* MorphLOD = Owner.GetMorphTargets()[*FoundIndex];
+				MorphLOD->GetMorphLODModels()[LODIndex].SectionIndices = MutableMesh.Morph.SurfacesInUsePerMorph[MorphIndex];
 			}
 		}
 
diff --git a/Engine/Plugins/Mutable/Source/MutableRuntime/Private/MuR/Mesh.cpp b/Engine/Plugins/Mutable/Source/MutableRuntime/Private/MuR/Mesh.cpp
--- a/Engine/Plugins/Mutable/Source/MutableRuntime/Private/MuR/Mesh.cpp
+++ b/Engine/Plugins/Mutable/Source/MutableRuntime/Private/MuR/Mesh.cpp
@@ -209,7 +209,7 @@ uint32 FMesh::GetId() const
 
 bool FMesh::HasMorphs() const
 {
-	return MorphDataBuffer.Num() && Morph.NumTotalBatches != 0;
+	return Morph.NumTotalBatches != 0;
 }
 
 
diff --git a/Engine/Plugins/Mutable/Source/MutableRuntime/Private/MuR/OpMeshMerge.cpp b/Engine/Plugins/Mutable/Source/MutableRuntime/Private/MuR/OpMeshMerge.cpp
--- a/Engine/Plugins/Mutable/Source/MutableRuntime/Private/MuR/OpMeshMerge.cpp
+++ b/Engine/Plugins/Mutable/Source/MutableRuntime/Private/MuR/OpMeshMerge.cpp
@@ -64,7 +64,7 @@ namespace OpMeshMergeMorphsInternal
 				}
 			}
 		}
-	
+
 		int32 NumMorphsToRemove = Algo::CountIf(Result.Morph.UsageFlagsPerMorph, 
 				[UsageFilter](EMorphUsageFlags Flags) { return !EnumHasAnyFlags(Flags, UsageFilter); });
 
@@ -80,73 +80,77 @@ namespace OpMeshMergeMorphsInternal
 				if (EnumHasAnyFlags(Result.Morph.UsageFlagsPerMorph[MorphIndex], UsageFilter))
 				{
 					int32 NumBatches = Result.Morph.BatchesPerMorph[MorphIndex];
-					
-					FMemory::Memmove(
-							Result.MorphDataBuffer.GetData() + AccumulatedDataOffsetInDwords,
-							Result.MorphDataBuffer.GetData() + Result.Morph.BatchStartOffsetPerMorph[MorphIndex] * NumBatchHeaderDwords,
-							NumBatches * NumBatchHeaderDwords * sizeof(uint32));
-
 					Result.Morph.BatchStartOffsetPerMorph[MorphIndex] = Result.Morph.NumTotalBatches;
 					Result.Morph.NumTotalBatches += NumBatches;
-					
-					AccumulatedDataOffsetInDwords += NumBatches * NumBatchHeaderDwords;
+				
+					if (Result.MorphDataBuffer.Num())
+					{
+						FMemory::Memmove(
+								Result.MorphDataBuffer.GetData() + AccumulatedDataOffsetInDwords,
+								Result.MorphDataBuffer.GetData() + Result.Morph.BatchStartOffsetPerMorph[MorphIndex] * NumBatchHeaderDwords,
+								NumBatches * NumBatchHeaderDwords * sizeof(uint32));
+
+						
+						AccumulatedDataOffsetInDwords += NumBatches * NumBatchHeaderDwords;
+					}
 				}
 				else
 				{
 					Result.Morph.BatchesPerMorph[MorphIndex] = 0;
 					Result.Morph.BatchStartOffsetPerMorph[MorphIndex] = Result.Morph.NumTotalBatches;
 				}
-				
 			}
 
-			// Move the data of the remaining batches and patch the offsets.
-			for (int32 MorphIndex = 0; MorphIndex < NumMorphs; ++MorphIndex)
+			if (Result.MorphDataBuffer.Num())
 			{
-				int32 NumBatches = Result.Morph.BatchesPerMorph[MorphIndex];
-
-				if (NumBatches == 0)
+				// Move the data of the remaining batches and patch the offsets.
+				for (int32 MorphIndex = 0; MorphIndex < NumMorphs; ++MorphIndex)
 				{
-					continue;
-				}
+					int32 NumBatches = Result.Morph.BatchesPerMorph[MorphIndex];
 
-				int32 BatchStartOffset = Result.Morph.BatchStartOffsetPerMorph[MorphIndex];
+					if (NumBatches == 0)
+					{
+						continue;
+					}
 
-				int32 MorphDeltasDataBeginInBytes = Result.MorphDataBuffer[BatchStartOffset * NumBatchHeaderDwords];
-				int32 OldToNewDeltaOffsetInBytes = MorphDeltasDataBeginInBytes - (AccumulatedDataOffsetInDwords * sizeof(uint32));
+					int32 BatchStartOffset = Result.Morph.BatchStartOffsetPerMorph[MorphIndex];
 
-				check(OldToNewDeltaOffsetInBytes >= 0);
+					int32 MorphDeltasDataBeginInBytes = Result.MorphDataBuffer[BatchStartOffset * NumBatchHeaderDwords];
+					int32 OldToNewDeltaOffsetInBytes = MorphDeltasDataBeginInBytes - (AccumulatedDataOffsetInDwords * sizeof(uint32));
 
-				for (int32 BatchIndex = 0; BatchIndex < NumBatches; ++BatchIndex)
-				{
-					Result.MorphDataBuffer[(BatchStartOffset + BatchIndex) * NumBatchHeaderDwords] -= OldToNewDeltaOffsetInBytes; 
-				}
-				
-				int32 LastBatchDataOffsetInBytes = Result.MorphDataBuffer[(BatchStartOffset + NumBatches - 1) * NumBatchHeaderDwords];
+					check(OldToNewDeltaOffsetInBytes >= 0);
 
-				TArrayView<uint32> LastHeaderDataView = TArrayView<uint32>(
-						&Result.MorphDataBuffer[(BatchStartOffset + NumBatches - 1) * NumBatchHeaderDwords],
-						NumBatchHeaderDwords);
+					for (int32 BatchIndex = 0; BatchIndex < NumBatches; ++BatchIndex)
+					{
+						Result.MorphDataBuffer[(BatchStartOffset + BatchIndex) * NumBatchHeaderDwords] -= OldToNewDeltaOffsetInBytes; 
+					}
+					
+					int32 LastBatchDataOffsetInBytes = Result.MorphDataBuffer[(BatchStartOffset + NumBatches - 1) * NumBatchHeaderDwords];
 
-				FDeltaBatchHeader LastBatchHeader;
-				UE::MorphTargetVertexCodec::ReadHeader(LastBatchHeader, LastHeaderDataView);
+					TArrayView<uint32> LastHeaderDataView = TArrayView<uint32>(
+							&Result.MorphDataBuffer[(BatchStartOffset + NumBatches - 1) * NumBatchHeaderDwords],
+							NumBatchHeaderDwords);
+
+					FDeltaBatchHeader LastBatchHeader;
+					UE::MorphTargetVertexCodec::ReadHeader(LastBatchHeader, LastHeaderDataView);
 
-				int32 MorphDeltaDataSizeInBytes = 
-					(LastBatchDataOffsetInBytes + CalculateBatchDwords(LastBatchHeader) * sizeof(uint32)) -
-					(MorphDeltasDataBeginInBytes - OldToNewDeltaOffsetInBytes);
+					int32 MorphDeltaDataSizeInBytes = 
+						(LastBatchDataOffsetInBytes + CalculateBatchDwords(LastBatchHeader) * sizeof(uint32)) -
+						(MorphDeltasDataBeginInBytes - OldToNewDeltaOffsetInBytes);
 
-				FMemory::Memmove(
-						Result.MorphDataBuffer.GetData() + AccumulatedDataOffsetInDwords,
-						Result.MorphDataBuffer.GetData() + (MorphDeltasDataBeginInBytes / sizeof(uint32)),
-						MorphDeltaDataSizeInBytes);	
+					FMemory::Memmove(
+							Result.MorphDataBuffer.GetData() + AccumulatedDataOffsetInDwords,
+							Result.MorphDataBuffer.GetData() + (MorphDeltasDataBeginInBytes / sizeof(uint32)),
+							MorphDeltaDataSizeInBytes);	
 
-				AccumulatedDataOffsetInDwords += MorphDeltaDataSizeInBytes / sizeof(uint32);
+					AccumulatedDataOffsetInDwords += MorphDeltaDataSizeInBytes / sizeof(uint32);
+				}
+				// Free unused data.
+				Result.MorphDataBuffer.SetNum(AccumulatedDataOffsetInDwords, EAllowShrinking::Yes);
 			}
+		}	
 		
-			// Free unused data.
-			Result.MorphDataBuffer.SetNum(AccumulatedDataOffsetInDwords, EAllowShrinking::Yes);
-		}
-
-		if (VertexIndexOffset != 0)
+		if (VertexIndexOffset != 0 && Result.MorphDataBuffer.Num())
 		{
 			for (int32 MorphIndex = 0; MorphIndex < NumMorphs; ++MorphIndex)
 			{
@@ -170,7 +174,6 @@ namespace OpMeshMergeMorphsInternal
 				}
 			}
 		}
-
 	}
 
 	void MergeMorphs(FMesh& Result, const FMesh& MeshA, const FMesh& MeshB, bool bMergeSurfaces, EMorphUsageFlags UsageFilter)
@@ -261,11 +264,14 @@ namespace OpMeshMergeMorphsInternal
 					int32 NumBatches = MeshA.Morph.BatchesPerMorph[MorphIndex];
 					Result.Morph.NumTotalBatches += NumBatches;
 
-					TConstArrayView<uint32> HeadersDataView(
-							MeshA.MorphDataBuffer.GetData() + MeshA.Morph.BatchStartOffsetPerMorph[MorphIndex]*NumBatchHeaderDwords,
-							NumBatchHeaderDwords*NumBatches);
+					if (MeshA.MorphDataBuffer.Num())
+					{
+						TConstArrayView<uint32> HeadersDataView(
+								MeshA.MorphDataBuffer.GetData() + MeshA.Morph.BatchStartOffsetPerMorph[MorphIndex]*NumBatchHeaderDwords,
+								NumBatchHeaderDwords*NumBatches);
 
-					MergedDataBufferSizeInDwords += ComputeMorphAllocationRequirementsInDwords(HeadersDataView);
+						MergedDataBufferSizeInDwords += ComputeMorphAllocationRequirementsInDwords(HeadersDataView);
+					}
 				}
 			}
 
@@ -276,11 +282,14 @@ namespace OpMeshMergeMorphsInternal
 					int32 NumBatches = MeshB.Morph.BatchesPerMorph[MorphIndex];
 					Result.Morph.NumTotalBatches += NumBatches;
 
-					TConstArrayView<uint32> HeadersDataView(
-							MeshB.MorphDataBuffer.GetData() + MeshB.Morph.BatchStartOffsetPerMorph[MorphIndex]*NumBatchHeaderDwords,
-							NumBatchHeaderDwords*NumBatches);
+					if (MeshB.MorphDataBuffer.Num())
+					{
+						TConstArrayView<uint32> HeadersDataView(
+								MeshB.MorphDataBuffer.GetData() + MeshB.Morph.BatchStartOffsetPerMorph[MorphIndex]*NumBatchHeaderDwords,
+								NumBatchHeaderDwords*NumBatches);
 
-					MergedDataBufferSizeInDwords += ComputeMorphAllocationRequirementsInDwords(HeadersDataView);
+						MergedDataBufferSizeInDwords += ComputeMorphAllocationRequirementsInDwords(HeadersDataView);
+					}
 				}
 			}
 
@@ -309,9 +318,9 @@ namespace OpMeshMergeMorphsInternal
 							Result.MorphDataBuffer.GetData() + BatchAccumulatedOffsetInDwords, 
 							MeshA.MorphDataBuffer.GetData() + MeshA.Morph.BatchStartOffsetPerMorph[MorphIndex]*NumBatchHeaderDwords,
 							MeshA.Morph.BatchesPerMorph[MorphIndex] * NumBatchHeaderDwords*sizeof(uint32));
+					
+					BatchAccumulatedOffsetInDwords += MeshA.Morph.BatchesPerMorph[MorphIndex] * NumBatchHeaderDwords;
 				}
-			
-				BatchAccumulatedOffsetInDwords += MeshA.Morph.BatchesPerMorph[MorphIndex] * NumBatchHeaderDwords;
 			}
 
 			int32 MeshBMorphIndex = ResultToBMorphNameMap[MorphIndex];
@@ -335,9 +344,9 @@ namespace OpMeshMergeMorphsInternal
 								Result.MorphDataBuffer.GetData() + BatchAccumulatedOffsetInDwords, 
 								MeshB.MorphDataBuffer.GetData() + MeshB.Morph.BatchStartOffsetPerMorph[MeshBMorphIndex]*NumBatchHeaderDwords,
 								MeshB.Morph.BatchesPerMorph[MeshBMorphIndex] * NumBatchHeaderDwords*sizeof(uint32));
+					
+						BatchAccumulatedOffsetInDwords += MeshB.Morph.BatchesPerMorph[MeshBMorphIndex] * NumBatchHeaderDwords;
 					}
-
-					BatchAccumulatedOffsetInDwords += MeshB.Morph.BatchesPerMorph[MeshBMorphIndex] * NumBatchHeaderDwords;
 				}
 			}
 		}
@@ -363,9 +372,9 @@ namespace OpMeshMergeMorphsInternal
 							Result.MorphDataBuffer.GetData() + BatchAccumulatedOffsetInDwords, 
 							MeshB.MorphDataBuffer.GetData() + MeshB.Morph.BatchStartOffsetPerMorph[MeshBMorphIndex]*NumBatchHeaderDwords,
 							MeshB.Morph.BatchesPerMorph[MeshBMorphIndex] * NumBatchHeaderDwords*sizeof(uint32));
+					
+					BatchAccumulatedOffsetInDwords += MeshB.Morph.BatchesPerMorph[MeshBMorphIndex] * NumBatchHeaderDwords;
 				}
-
-				BatchAccumulatedOffsetInDwords += MeshB.Morph.BatchesPerMorph[MeshBMorphIndex] * NumBatchHeaderDwords;
 			}
 		}
 
diff --git a/Engine/Plugins/Mutable/Source/MutableRuntime/Private/MuR/OpMeshMorph.cpp b/Engine/Plugins/Mutable/Source/MutableRuntime/Private/MuR/OpMeshMorph.cpp
--- a/Engine/Plugins/Mutable/Source/MutableRuntime/Private/MuR/OpMeshMorph.cpp
+++ b/Engine/Plugins/Mutable/Source/MutableRuntime/Private/MuR/OpMeshMorph.cpp
@@ -19,7 +19,7 @@ void UE::Mutable::Private::MeshMorph(FMesh* Mesh, const FName& MorphName, const
 
 	using namespace UE::MorphTargetVertexCodec;
 
-	if (!Mesh)
+	if (!(Mesh && Mesh->MorphDataBuffer.Num()))
 	{
 		return;
 	}
diff --git a/Engine/Plugins/Mutable/Source/MutableRuntime/Private/MuR/OpMeshRemove.cpp b/Engine/Plugins/Mutable/Source/MutableRuntime/Private/MuR/OpMeshRemove.cpp
--- a/Engine/Plugins/Mutable/Source/MutableRuntime/Private/MuR/OpMeshRemove.cpp
+++ b/Engine/Plugins/Mutable/Source/MutableRuntime/Private/MuR/OpMeshRemove.cpp
@@ -107,7 +107,7 @@ namespace UE::Mutable::Private
 
 		using namespace UE::MorphTargetVertexCodec;
 
-		if (!Result.HasMorphs())
+		if (!Result.MorphDataBuffer.Num())
 		{
 			return;
 		}
