[Mutable] Fix signed/unsigned compare compilation error.

#rnx
#Synced-CL 51539466

diff --git a/Engine/Plugins/Mutable/Source/CustomizableObject/Private/MuCO/CustomizableObjectSystem.cpp b/Engine/Plugins/Mutable/Source/CustomizableObject/Private/MuCO/CustomizableObjectSystem.cpp
--- a/Engine/Plugins/Mutable/Source/CustomizableObject/Private/MuCO/CustomizableObjectSystem.cpp
+++ b/Engine/Plugins/Mutable/Source/CustomizableObject/Private/MuCO/CustomizableObjectSystem.cpp
@@ -3765,7 +3765,7 @@ void ReconstructMorphTargetsFromMeshCompressedData(const UE::Mutable::Private::F
 			const int32 BatchStartOffset = Mesh.Morph.BatchStartOffsetPerMorph[MorphIndex];
 			const int32 MorphNumBatches  = Mesh.Morph.BatchesPerMorph[MorphIndex];
 
-			check((BatchStartOffset + MorphNumBatches)*NumBatchHeaderDwords <= Mesh.MorphDataBuffer.Num());
+			check((BatchStartOffset + MorphNumBatches)*NumBatchHeaderDwords <= (uint32)Mesh.MorphDataBuffer.Num());
 			TConstArrayView<uint32> MorphHeadersView = TConstArrayView<uint32>(
 				Mesh.MorphDataBuffer.GetData() + BatchStartOffset * NumBatchHeaderDwords,
 				MorphNumBatches * NumBatchHeaderDwords);
