[mutable] Removed the FInstance returned by the CodeRunnerBegin as it is being overridden by the CodeRunner that runs downstream.
- Since we do not need the FInstance to be set up at this point, its computation has been removed from the CodeRunnerBegin object.
- Handling of strings and materials has also been removed from the operation of the class, as they are not required for our use case (set up the FLiveInstance FPassthroughObjectLoader).
- Added some missing types into the mutable code viewer.
- Made the MeshCache thread safe.

#rnx
#Synced-CL 51580945

diff --git a/Engine/Plugins/Mutable/Source/CustomizableObject/Internal/MuCO/CustomizableObjectPrivate.h b/Engine/Plugins/Mutable/Source/CustomizableObject/Internal/MuCO/CustomizableObjectPrivate.h
--- a/Engine/Plugins/Mutable/Source/CustomizableObject/Internal/MuCO/CustomizableObjectPrivate.h
+++ b/Engine/Plugins/Mutable/Source/CustomizableObject/Internal/MuCO/CustomizableObjectPrivate.h
@@ -136,7 +136,7 @@ struct FMutableModelParameterProperties
 class FMeshCache
 {
 public:
-	UCustomizableObjectSkeletalMesh* Get(const TArray<UE::Mutable::Private::FMeshId>& Key);
+	TWeakObjectPtr<UCustomizableObjectSkeletalMesh> Get(const TArray<UE::Mutable::Private::FMeshId>& Key);
 
 	void Add(const TArray<UE::Mutable::Private::FMeshId>& Key, UCustomizableObjectSkeletalMesh* Value);
 
@@ -149,6 +149,8 @@ private:
 
 public:
 	typedef FMeshCacheRegistry::FHandle FId;
+	
+	UE::FMutex Lock;
 };
 
 
@@ -168,8 +170,6 @@ public:
 
 		friend uint32 GetTypeHash(const FId& Key);
 	};
-
-	UE::FMutex Lock;
 	
 	TWeakObjectPtr<UTexture2D> Get(const FId& Key);
 
@@ -181,6 +181,8 @@ public:
 
 private:
 	TMap<FId, TWeakObjectPtr<UTexture2D>> GeneratedTextures;
+	
+	UE::FMutex Lock;
 };
 
 
diff --git a/Engine/Plugins/Mutable/Source/CustomizableObject/Private/MuCO/CustomizableObject.cpp b/Engine/Plugins/Mutable/Source/CustomizableObject/Private/MuCO/CustomizableObject.cpp
--- a/Engine/Plugins/Mutable/Source/CustomizableObject/Private/MuCO/CustomizableObject.cpp
+++ b/Engine/Plugins/Mutable/Source/CustomizableObject/Private/MuCO/CustomizableObject.cpp
@@ -1819,10 +1819,12 @@ void UCustomizableObjectPrivate::ApplyStateForcedValuesToParameters(FCustomizabl
 //-------------------------------------------------------------------------------------------------
 //-------------------------------------------------------------------------------------------------
 
-UCustomizableObjectSkeletalMesh* FMeshCache::Get(const TArray<UE::Mutable::Private::FMeshId>& Key)
+TWeakObjectPtr<UCustomizableObjectSkeletalMesh> FMeshCache::Get(const TArray<UE::Mutable::Private::FMeshId>& Key)
 {
+	UE::TScopeLock<UE::FMutex> LockGuard(Lock);
+	
 	const TWeakObjectPtr<UCustomizableObjectSkeletalMesh>* Result = GeneratedMeshes->Find(Key);
-	return Result ? Result->Get() : nullptr;
+	return Result ? *Result : TWeakObjectPtr<UCustomizableObjectSkeletalMesh>();
 }
 
 
@@ -1833,12 +1835,17 @@ void FMeshCache::Add(const TArray<UE::Mutable::Private::FMeshId>& Key, UCustomiz
 		return;
 	}
 	
-	Value->MeshCacheId = GeneratedMeshes->Add(Key, TWeakObjectPtr(Value));
+	{
+		UE::TScopeLock<UE::FMutex> LockGuard(Lock);
+		
+		Value->MeshCacheId = GeneratedMeshes->Add(Key, TWeakObjectPtr(Value));
+	}
 }
 
 
 void FMeshCache::Clear()
 {
+	UE::TScopeLock<UE::FMutex> LockGuard(Lock);
 	GeneratedMeshes = MakeShared<FMeshCacheRegistry>();
 }
 
diff --git a/Engine/Plugins/Mutable/Source/CustomizableObject/Private/MuCO/CustomizableObjectInstance.cpp b/Engine/Plugins/Mutable/Source/CustomizableObject/Private/MuCO/CustomizableObjectInstance.cpp
--- a/Engine/Plugins/Mutable/Source/CustomizableObject/Private/MuCO/CustomizableObjectInstance.cpp
+++ b/Engine/Plugins/Mutable/Source/CustomizableObject/Private/MuCO/CustomizableObjectInstance.cpp
@@ -1769,9 +1769,13 @@ bool UCustomizableInstancePrivate::UpdateSkeletalMesh_PostBeginUpdate0(UCustomiz
 
 		if (Context->bUseMeshCache)
 		{
-			if (USkeletalMesh* CachedMesh = CustomizableObject->GetPrivate()->MeshCache.Get(ComponentDescriptor->MeshDescriptors))
+			TWeakObjectPtr<UCustomizableObjectSkeletalMesh> MeshCacheEntry = CustomizableObject->GetPrivate()->MeshCache.Get(
+				ComponentDescriptor->MeshDescriptors);
+			
+			TObjectPtr<UCustomizableObjectSkeletalMesh> Mesh = MeshCacheEntry.Get();
+			if (IsValid(Mesh))
 			{
-				SkeletalMeshes.Add(ComponentName, CachedMesh);
+				SkeletalMeshes.Add(ComponentName, Mesh);
 				continue;
 			}
 		}
diff --git a/Engine/Plugins/Mutable/Source/CustomizableObject/Private/MuCO/CustomizableObjectSystem.cpp b/Engine/Plugins/Mutable/Source/CustomizableObject/Private/MuCO/CustomizableObjectSystem.cpp
--- a/Engine/Plugins/Mutable/Source/CustomizableObject/Private/MuCO/CustomizableObjectSystem.cpp
+++ b/Engine/Plugins/Mutable/Source/CustomizableObject/Private/MuCO/CustomizableObjectSystem.cpp
@@ -1568,7 +1568,8 @@ namespace impl
 	
 	void FixLODs(const TSharedRef<FUpdateContextPrivate>& Operation)
 	{
-		for (int32 ComponentIndex = 0; ComponentIndex < Operation->MutableInstance->GetComponentCount(); ++ComponentIndex)
+		const int32 ComponentCount =  Operation->MutableInstance->GetComponentCount();
+		for (int32 ComponentIndex = 0; ComponentIndex < ComponentCount; ++ComponentIndex)
 		{
 			const UE::Mutable::Private::FComponentId ComponentId = Operation->MutableInstance->GetComponentId(ComponentIndex);
 			const FName& ComponentName = Operation->ComponentNames[ComponentId];
@@ -1870,6 +1871,51 @@ namespace impl
 		
 		OperationData->MutableInstance = Instance;
 
+		FixLODs(OperationData);
+
+		// Setup Component and mesh descriptors
+		{
+			OperationData->InitComponentDescriptors(OperationData->Object->GetComponentCount());
+            		
+			const TMap<FName, uint8>& NumLODsAvailable = OperationData->NumLODsAvailable;
+	
+			for (int32 ComponentIndex = 0; ComponentIndex < OperationData->MutableInstance->GetComponentCount(); ++ComponentIndex)
+			{
+				UE::Mutable::Private::FComponentId ComponentId = OperationData->MutableInstance->GetComponentId(ComponentIndex);
+				FComponentDescriptor* ComponentDescriptor = OperationData->GetComponentDescriptor(ComponentId);
+				check(ComponentDescriptor);
+	
+				if (!ComponentDescriptor)
+				{
+					continue;
+				}
+	
+				ComponentDescriptor->MeshDescriptors.Init({}, MAX_MESH_LOD_COUNT);
+	
+				const FName& ComponentName = OperationData->ComponentNames[ComponentId];
+	
+				for (int32 LODIndex = OperationData->FirstLODAvailable[ComponentName]; LODIndex < NumLODsAvailable[ComponentName]; ++LODIndex)
+				{
+					ComponentDescriptor->MeshDescriptors[LODIndex] = OperationData->MutableInstance->GetMeshId(ComponentIndex, LODIndex);
+				}
+			}
+		}
+		
+
+		if (OperationData->bUseMeshCache)
+		{
+			for (const FComponentDescriptor& ComponentDescriptor : OperationData->GetComponentDescriptors())
+			{
+				TWeakObjectPtr<UCustomizableObjectSkeletalMesh> MeshCacheEntry = OperationData->Object->GetPrivate()->MeshCache.Get(ComponentDescriptor.MeshDescriptors);
+				
+				bool bWasPinningSuccessful = false;
+				if (TStrongObjectPtr<UCustomizableObjectSkeletalMesh> CachedMesh = MeshCacheEntry.TryPin(bWasPinningSuccessful))
+				{
+					OperationData->Objects.Emplace(CachedMesh);
+				}
+			}
+		}
+		
 		OperationData->ExtensionData = Instance->ExtensionData;
 		
 		const bool bForceGeometryGeneration = CVarForceGeometryOnFirstGeneration.GetValueOnAnyThread();
@@ -2643,49 +2689,11 @@ namespace impl
 				Instance->GetPrivate()->LiveInstance = nullptr;
 			}
 			
-			Operation->MutableInstance = MutableSystem->BeginUpdate_GameThread(
+			MutableSystem->BeginUpdate_GameThread(
 				Operation->LiveInstance.ToSharedRef(),
 				UE::Mutable::Private::FSystem::AllLODs);
 		}
 		
-		FixLODs(Operation);
-		
-		Operation->InitComponentDescriptors(Object->GetComponentCount());
-
-		const TMap<FName, uint8>& NumLODsAvailable = Operation->NumLODsAvailable;
-
-		for (int32 ComponentIndex = 0; ComponentIndex < Operation->MutableInstance->GetComponentCount(); ++ComponentIndex)
-		{
-			UE::Mutable::Private::FComponentId ComponentId = Operation->MutableInstance->GetComponentId(ComponentIndex);
-			FComponentDescriptor* ComponentDescriptor = Operation->GetComponentDescriptor(ComponentId);
-			check(ComponentDescriptor);
-
-			if (!ComponentDescriptor)
-			{
-				continue;
-			}
-
-			ComponentDescriptor->MeshDescriptors.Init({}, MAX_MESH_LOD_COUNT);
-
-			const FName& ComponentName = Operation->ComponentNames[ComponentId];
-
-			for (int32 LODIndex = Operation->FirstLODAvailable[ComponentName]; LODIndex < NumLODsAvailable[ComponentName]; ++LODIndex)
-			{
-				ComponentDescriptor->MeshDescriptors[LODIndex] = Operation->MutableInstance->GetMeshId(ComponentIndex, LODIndex);
-			}
-		}
-
-		if (Operation->bUseMeshCache)
-		{
-			for (const FComponentDescriptor& ComponentDescriptor : Operation->GetComponentDescriptors())
-			{
-				if (UCustomizableObjectSkeletalMesh* CachedMesh = ObjectPrivate->MeshCache.Get(ComponentDescriptor.MeshDescriptors))
-				{
-					Operation->Objects.Emplace(CachedMesh);
-				}
-			}
-		}
-		
 		TArray<	UE::Tasks::FTask> Prerequisites;
 		Prerequisites.Add(CacheRuntimeTexturesEvent);
 		
diff --git a/Engine/Plugins/Mutable/Source/CustomizableObjectEditor/Private/MuCOE/SMutableGraphViewer.cpp b/Engine/Plugins/Mutable/Source/CustomizableObjectEditor/Private/MuCOE/SMutableGraphViewer.cpp
--- a/Engine/Plugins/Mutable/Source/CustomizableObjectEditor/Private/MuCOE/SMutableGraphViewer.cpp
+++ b/Engine/Plugins/Mutable/Source/CustomizableObjectEditor/Private/MuCOE/SMutableGraphViewer.cpp
@@ -5,9 +5,7 @@
 #include "DesktopPlatformModule.h"
 #include "EditorDirectories.h"
 #include "Framework/MultiBox/MultiBoxBuilder.h"
-#include "Framework/Views/TableViewMetadata.h"
 #include "IDesktopPlatform.h"
-#include "Misc/Paths.h"
 #include "MuCOE/CustomizableObjectCompileRunnable.h"
 #include "MuCOE/CustomizableObjectEditorStyle.h"
 #include "MuCOE/SMutableCodeViewer.h"
@@ -16,12 +14,12 @@
 #include "Widgets/Input/STextComboBox.h"
 #include "Widgets/Input/SNumericDropDown.h"
 #include "Widgets/Views/STreeView.h"
-#include "ScopedTransaction.h"
 #include "MuT/NodeColourConstant.h"
 #include "MuT/NodeColourFromScalars.h"
 #include "MuT/NodeColourParameter.h"
 #include "MuT/NodeColourSampleImage.h"
 #include "MuT/NodeColourSwitch.h"
+#include "MuT/NodeColourTable.h"
 #include "MuT/NodeComponentEdit.h"
 #include "MuT/NodeComponentSwitch.h"
 #include "MuT/NodeImageFormat.h"
@@ -53,6 +51,7 @@
 #include "MuT/NodeMeshMakeMorph.h"
 #include "MuT/NodeMeshMorph.h"
 #include "MuT/NodeMeshRemoveMesh.h"
+#include "MuT/NodeMeshReshape.h"
 #include "MuT/NodeMeshTable.h"
 #include "MuT/NodeMeshTransform.h"
 #include "MuT/NodeMeshTransformInMesh.h"
@@ -66,9 +65,9 @@
 #include "MuT/NodeScalarCurve.h"
 #include "MuT/NodeScalarSwitch.h"
 #include "MuT/NodeScalarTable.h"
-#include "MuT/NodeObjectNew.h"
 #include "MuT/NodeModifierMeshTransformInMesh.h"
 #include "MuT/NodeModifierMeshTransformWithBone.h"
+#include "MuT/NodeScalarArithmeticOperation.h"
 #include "MuT/NodeSkeletalMeshNew.h"
 #include "Widgets/MutableExpanderArrow.h"
 
@@ -960,6 +959,26 @@ void SMutableGraphViewer::GetChildrenForInfo(TSharedPtr<FMutableGraphTreeElement
 		AddChildFunc(NodeMeshTransformInMesh->MatrixNode.get(), TEXT("TRANSFORM"));
 	}
 	
+	else if (ParentNodeType == UE::Mutable::Private::NodeColourTable::GetStaticType())
+	{
+		// Nothing to show
+	}
+	
+	else if (ParentNodeType == UE::Mutable::Private::NodeMeshReshape::GetStaticType())
+	{
+		const UE::Mutable::Private::NodeMeshReshape* NodeMeshReshape = StaticCast<UE::Mutable::Private::NodeMeshReshape*>(ParentNode);
+		AddChildFunc(NodeMeshReshape->BaseMesh.get(), TEXT("BASE MESH"));
+		AddChildFunc(NodeMeshReshape->BaseShape.get(), TEXT("BASE SHAPE"));
+		AddChildFunc(NodeMeshReshape->TargetShape.get(), TEXT("TGT SHAPE"));
+	}
+	
+	else if (ParentNodeType == UE::Mutable::Private::NodeScalarArithmeticOperation::GetStaticType())
+	{
+		const UE::Mutable::Private::NodeScalarArithmeticOperation* NodeScalarArithmeticOperation = StaticCast<UE::Mutable::Private::NodeScalarArithmeticOperation*>(ParentNode);
+		AddChildFunc(NodeScalarArithmeticOperation->A.get(), TEXT("A"));
+		AddChildFunc(NodeScalarArithmeticOperation->B.get(), TEXT("B"));
+	}
+	
 	else
 	{
 		UE_LOG(LogMutable,Error,TEXT("The node of type %d has not been implemented, so its children won't be added to the tree."), int32(ParentNodeType->Type));
diff --git a/Engine/Plugins/Mutable/Source/MutableRuntime/Internal/MuR/System.h b/Engine/Plugins/Mutable/Source/MutableRuntime/Internal/MuR/System.h
--- a/Engine/Plugins/Mutable/Source/MutableRuntime/Internal/MuR/System.h
+++ b/Engine/Plugins/Mutable/Source/MutableRuntime/Internal/MuR/System.h
@@ -308,7 +308,7 @@ namespace UE::Mutable::Private
         //! the same instanceID parameter.
 		//!
 		//! Very fast. Can be called in the Game Thread.
-		UE_API TManagedPtr<const FInstance> BeginUpdate_GameThread(const TSharedRef<FLiveInstance>& LiveInstance, uint32 LodMask);
+		UE_API void BeginUpdate_GameThread(const TSharedRef<FLiveInstance>& LiveInstance, uint32 LodMask);
 
 		/** Potentially expensive. Must be called in the Mutable Thread. */
 		UE_API TManagedPtr<const FInstance> BeginUpdate_MutableThread(const TSharedRef<FLiveInstance>& LiveInstance,
diff --git a/Engine/Plugins/Mutable/Source/MutableRuntime/Private/MuR/CodeRunnerBegin.cpp b/Engine/Plugins/Mutable/Source/MutableRuntime/Private/MuR/CodeRunnerBegin.cpp
--- a/Engine/Plugins/Mutable/Source/MutableRuntime/Private/MuR/CodeRunnerBegin.cpp
+++ b/Engine/Plugins/Mutable/Source/MutableRuntime/Private/MuR/CodeRunnerBegin.cpp
@@ -21,10 +21,6 @@ namespace UE::Mutable::Private
 		LodMask(InLodMask),
 		Executed(InLiveInstance->Model->GetProgram().NumOps),
 		ProgramCache(InLiveInstance->Cache),
-		MeshIdRegistry(InLiveInstance->MeshIdRegistry),
-		ImageIdRegistry(InLiveInstance->ImageIdRegistry),
-		MaterialIdRegistry(InLiveInstance->MaterialIdRegistry),
-		SkeletalMeshIdRegistry(InLiveInstance->SkeletalMeshIdRegistry),
 		PassthroughObjectLoader(InLiveInstance->PassthroughObjectLoader)
 	{
 	}
@@ -148,39 +144,7 @@ namespace UE::Mutable::Private
 	}
 
 
-	void CodeRunnerBegin::StoreInstance(const FCacheAddressInline& To, const TManagedPtr<const FInstance>& Value)
-	{
-		if (!To.At)
-		{
-			return;
-		}
-
-		Executed.Add(To);
 
-		FStackValue StackValue;
-		StackValue.Set<TManagedPtr<const FInstance>>(Value);
-		Stack.Push(StackValue);
-
-		ResultsInstance.Add(To, Value);
-	}
-
-
-	void CodeRunnerBegin::StoreString(const FCacheAddressInline& To, const TManagedPtr<const String>& Value)
-	{
-		if (!To.At)
-		{
-			return;
-		}
-		
-		Executed.Add(To);
-
-		FStackValue StackValue;
-		StackValue.Set<TManagedPtr<const String>>(Value);
-		Stack.Push(StackValue);
-
-		ResultsString.Add(To, Value);
-	}
-	
 
 	void CodeRunnerBegin::StoreNone(const FCacheAddressInline& To)
 	{
@@ -193,23 +157,6 @@ namespace UE::Mutable::Private
 	}
 
 
-	void CodeRunnerBegin::StoreMaterial(const FCacheAddressInline& To, const TManagedPtr<const FMaterial>& Value)
-	{
-		if (!To.At)
-		{
-			return;
-		}
-		
-		Executed.Add(To);
-
-		FStackValue StackValue;
-		StackValue.Set<TManagedPtr<const FMaterial>>(Value);
-		Stack.Push(StackValue);
-
-		ResultsMaterial.Add(To, Value);
-	}
-
-
 	int32 CodeRunnerBegin::LoadInt(const FCacheAddressInline& To)
 	{
 		if (!To.At)
@@ -243,41 +190,8 @@ namespace UE::Mutable::Private
 	}
 
 
-	TManagedPtr<const FInstance> CodeRunnerBegin::LoadInstance(const FCacheAddressInline& To)
-	{
-		if (!To.At)
-		{
-			return nullptr;	
-		}
-			
-		return Stack.Pop().Get<TManagedPtr<const FInstance>>();
-	}
-
-
-	TManagedPtr<const String> CodeRunnerBegin::LoadString(const FCacheAddressInline& To)
-	{
-		if (!To.At)
-		{
-			return nullptr;	
-		}
-		
-		return Stack.Pop().Get<TManagedPtr<const String>>();
-	}
-	
-
-	TManagedPtr<const FMaterial> CodeRunnerBegin::LoadMaterial(const FCacheAddressInline& To)
-	{
-		if (!To.At)
+	void CodeRunnerBegin::RunCode(const FScheduledOpInline& Root)
 		{
-			return nullptr;	
-		}
-		
-		return Stack.Pop().Get<TManagedPtr<const FMaterial>>();
-	}
-
-
-	TManagedPtr<const FInstance> CodeRunnerBegin::RunCode(const FScheduledOpInline& Root)
-	{
 		MUTABLE_CPUPROFILER_SCOPE(CodeRunnerBegin::RunCode);
 		
 		PushOp(Root);
@@ -334,46 +248,6 @@ namespace UE::Mutable::Private
 						break;
 					}
 					
-				case EDataType::Instance:
-					{
-						if (Item.bEvaluate)
-						{
-							if (TManagedPtr<const FInstance>* Result = ResultsInstance.Find(Item))
-							{
-								FStackValue Value;
-								Value.Set<TManagedPtr<const FInstance>>(*Result);
-					
-								Stack.Push(Value);
-							}
-							else
-							{
-								bExecute = true;
-							}
-						}
-						
-						break;
-					}
-					
-				case EDataType::String:
-					{
-						if (Item.bEvaluate)
-						{
-							if (TManagedPtr<const String>* Result = ResultsString.Find(Item))
-							{
-								FStackValue Value;
-								Value.Set<TManagedPtr<const String>>(*Result);
-					
-								Stack.Push(Value);
-							}
-							else
-							{
-								bExecute = true;
-							}
-						}
-						
-						break;
-					}
-
 				case EDataType::Scalar:
 					{
 						if (Item.bEvaluate)
@@ -393,14 +267,7 @@ namespace UE::Mutable::Private
 						
 						break;
 					}
-				case EDataType::Material:
-					{
-						FStackValue Value;
-						Value.Set<TManagedPtr<const FMaterial>>(ResultsMaterial[Item]);
 
-						Stack.Push(Value);
-						break;
-					}
 					
 				case EDataType::ExtensionData:
 				case EDataType::Color:		
@@ -410,6 +277,9 @@ namespace UE::Mutable::Private
 				case EDataType::Layout:		
 				case EDataType::InstancedStruct:		
 				case EDataType::SkeletalMesh:		
+				case EDataType::Instance:
+				case EDataType::String:
+				case EDataType::Material:
 					break;
 
 				default:
@@ -433,6 +303,7 @@ namespace UE::Mutable::Private
 			case EOpType::MA_PARAMETER:
 			case EOpType::PR_PARAMETER:
 			case EOpType::SK_PARAMETER:
+			case EOpType::MI_PARAMETER:
 			case EOpType::IM_PARAMETER:
 			case EOpType::IM_PARAMETER_FROM_MATERIAL:
 			case EOpType::IS_PARAMETER:
@@ -443,12 +314,6 @@ namespace UE::Mutable::Private
 					break;
 				}
 
-			case EOpType::MI_PARAMETER:
-				{
-					StoreMaterial(Item,MakeManaged<const FMaterial>());
-					break;
-				}
-			
 			case EOpType::IN_CONDITIONAL:
 				check(Item.bEvaluate);
 
@@ -499,28 +364,19 @@ namespace UE::Mutable::Private
 									StoreInt(Item, LoadInt(Item.CustomState));
 									break;
 							
-								case EDataType::String:
-									StoreString(Item, LoadString(Item.CustomState));
-									break;
-
-								case EDataType::Instance:
-									StoreInstance(Item, LoadInstance(Item.CustomState));
-									break;
-
 								case EDataType::Scalar:		
 									StoreFloat(Item, LoadFloat(Item.CustomState));
 									break;
 							
-								case EDataType::Material:
-									StoreMaterial(Item, LoadMaterial(Item.CustomState));
-									break;
-
 								case EDataType::ExtensionData:
 								case EDataType::Color:		
 								case EDataType::Projector:	
 								case EDataType::Mesh:		
 								case EDataType::Image:		
 								case EDataType::Layout:		
+								case EDataType::String:
+								case EDataType::Instance:
+								case EDataType::Material:
 									StoreNone(Item);
 									break;
 
@@ -594,22 +450,10 @@ namespace UE::Mutable::Private
 										StoreInt(Item, 0);
 										break;
 								
-									case EDataType::Instance:
-										StoreInstance(Item, nullptr);
-										break;
-
-									case EDataType::String:
-										StoreString(Item, nullptr);
-										break;
-									
 									case EDataType::Scalar:
 										StoreFloat(Item, 0.0f);
 										break;
 
-									case EDataType::Material:
-										StoreMaterial(Item, MakeManaged<const FMaterial>());
-										break;
-									
 									case EDataType::ExtensionData:
 									case EDataType::Color:			
 									case EDataType::Projector:		
@@ -617,6 +461,9 @@ namespace UE::Mutable::Private
 									case EDataType::Image:			
 									case EDataType::Layout:			
 									case EDataType::InstancedStruct:			
+									case EDataType::Instance:
+									case EDataType::String:
+									case EDataType::Material:
 										StoreNone(Item);
 										break;
 								
@@ -706,22 +553,10 @@ namespace UE::Mutable::Private
 									StoreInt(Item, LoadInt(Item.CustomState));
 									break;
 
-								case EDataType::Instance:
-									StoreInstance(Item, LoadInstance(Item.CustomState));
-									break;
-								
-								case EDataType::String:
-									StoreString(Item, LoadString(Item.CustomState));
-									break;
-								
 								case EDataType::Scalar:
 									StoreFloat(Item, LoadFloat(Item.CustomState));
 									break;
 
-								case EDataType::Material:
-									StoreMaterial(Item, LoadMaterial(Item.CustomState));
-									break;
-
 								case EDataType::ExtensionData:
 								case EDataType::Color:
 								case EDataType::Projector:
@@ -729,6 +564,9 @@ namespace UE::Mutable::Private
 								case EDataType::Image:
 								case EDataType::Layout:
 								case EDataType::InstancedStruct:
+								case EDataType::Instance:
+								case EDataType::String:
+								case EDataType::Material:
 									StoreNone(Item);
 									break;
 
@@ -759,44 +597,12 @@ namespace UE::Mutable::Private
 					{
 					case 0:
 						{
-							PushOp(FScheduledOpInline(Item, 1));
 							PushOp(FScheduledOpInline(Args.instance, Item));
 							PushOp(FScheduledOpInline(Args.value, Item, false));
 							
 							break;
 						}
 
-					case 1:
-						{
-							TManagedPtr<const FInstance> InInstance;
-							if (Args.instance)
-							{
-								InInstance = LoadInstance(Args.instance);
-							}
-
-							TManagedPtr<FInstance> Instance;
-							if (InInstance)
-							{
-								Instance = UE::Mutable::Private::CloneOrTakeOver<FInstance>(MoveTemp(InInstance));
-							}
-							else
-							{
-								Instance = MakeManaged<FInstance>();
-							}
-							
-							if (Args.value)
-							{
-								FMeshId MeshId = FSystem::GetMeshId(MeshIdRegistry, *Model.Get(), *Params.Get(), Args.value);
-								FOperation::ADDRESS NameAd = Args.name;
-								check(NameAd < (uint32)Program.ConstantStrings.Num());
-								const FString& Name = Program.ConstantStrings[NameAd];
-								Instance->SetMesh(0, 0, MeshId, FName(Name));
-							}
-							
-							StoreInstance(Item, Instance);
-							break;
-						}
-
 					default:
 						unimplemented();
 					}
@@ -813,44 +619,12 @@ namespace UE::Mutable::Private
 					{
 					case 0:
 						{
-							PushOp(FScheduledOpInline(Item, 1));
 							PushOp(FScheduledOpInline(Args.Instance, Item));
 							PushOp(FScheduledOpInline(Args.Material, Item, false));
 							
 							break;
 						}
 
-					case 1:
-						{
-							TManagedPtr<const FInstance> InInstance;
-							if (Args.Instance)
-							{
-								InInstance = LoadInstance(Args.Instance);
-							}
-
-							TManagedPtr<FInstance> Instance;
-							if (InInstance)
-							{
-								Instance = UE::Mutable::Private::CloneOrTakeOver<FInstance>(MoveTemp(InInstance));
-							}
-							else
-							{
-								Instance = MakeManaged<FInstance>();
-							}
-
-
-							if (Args.Material)
-							{
-        						FMaterialId MaterialId = FSystem::GetMaterialId(MaterialIdRegistry, *Model.Get(), *Params.Get(), Args.Material);
-								Instance->SetMaterialId(0, 0, 0, MaterialId);
-
-								LoadMaterial(Args.Material);
-							}
-							
-							StoreInstance( Item, Instance);
-							break;
-						}
-
 					default:
 						unimplemented();
 					}
@@ -867,53 +641,12 @@ namespace UE::Mutable::Private
 					{
 					case 0:
 						{
-							PushOp(FScheduledOpInline(Item, 1));
 							PushOp(FScheduledOpInline(Args.Instance, Item));
 							PushOp(FScheduledOpInline(Args.Material, Item, false));
 
 							break;
 						}
 
-					case 1:
-						{
-							TManagedPtr<const FInstance> InInstance;
-							if (Args.Instance)
-							{
-								InInstance = LoadInstance(Args.Instance);
-							}
-
-							TManagedPtr<FInstance> Instance;
-							if (InInstance)
-							{
-								Instance = UE::Mutable::Private::CloneOrTakeOver<FInstance>(MoveTemp(InInstance));
-							}
-							else
-							{
-								Instance = MakeManaged<FInstance>();
-							}
-
-							if (Args.Material)
-							{
-								const FMaterialId MaterialId = FSystem::GetMaterialId(MaterialIdRegistry,  *Model.Get(),  *Params.Get(), Args.Material);
-        					
-								const FName SlotName = Args.SlotName == TNumericLimits<uint32>::Max() ? FName() : FName(Program.ConstantNames[Args.SlotName]);
-								if (SlotName.IsNone())
-								{
-									Instance->SetOverlayMaterialId(0, MaterialId);
-								}
-								else
-								{
-									Instance->AddOverlayMaterial(0, SlotName, MaterialId);
-								}
-								
-								// Pop Material
-								LoadMaterial(Args.Material);
-							}
-							
-							StoreInstance( Item, Instance);
-							break;
-						}
-
 					default:
 						unimplemented();
 					}
@@ -928,45 +661,12 @@ namespace UE::Mutable::Private
 					{
 					case 0:
 						{
-							PushOp(FScheduledOpInline(Item, 1));
 							PushOp(FScheduledOpInline(Args.Instance, Item));
 							PushOp(FScheduledOpInline(Args.Material, Item, false));
 
 							break;
 						}
 
-					case 1:
-						{
-							TManagedPtr<const FInstance> InInstance;
-							if (Args.Instance)
-							{
-								InInstance = LoadInstance(Args.Instance);
-							}
-
-							TManagedPtr<FInstance> Instance;
-							if (InInstance)
-							{
-								Instance = UE::Mutable::Private::CloneOrTakeOver<FInstance>(MoveTemp(InInstance));
-							}
-							else
-							{
-								Instance = MakeManaged<FInstance>();
-							}
-
-							if (Args.Material)
-							{
-								FMaterialId MaterialId = FSystem::GetMaterialId(MaterialIdRegistry, *Model.Get(), *Params.Get(), Args.Material);
-								const FString& SlotName = Program.ConstantStrings[Args.SlotName];
-								Instance->AddOverrideMaterial(0, FName(SlotName), MaterialId);
-
-								// Pop Material
-								LoadMaterial(Args.Material);
-							}
-							
-							StoreInstance( Item, Instance);
-							break;
-						}
-
 					default:
 						unimplemented();
 					}
@@ -983,50 +683,12 @@ namespace UE::Mutable::Private
 					{
 					case 0:
 						{
-							PushOp(FScheduledOpInline(Item, 1));
 							PushOp(FScheduledOpInline(Args.Instance, Item));
 							PushOp(FScheduledOpInline(Args.Value, Item));
 							
 							break;
 						}
 
-					case 1:
-						{
-							TManagedPtr<const FInstance> InInstance;
-							if (Args.Instance)
-							{
-								InInstance = LoadInstance(Args.Instance);
-							}
-
-							TManagedPtr<FInstance> Instance;
-							if (InInstance)
-							{
-								Instance = UE::Mutable::Private::CloneOrTakeOver<FInstance>(MoveTemp(InInstance));
-							}
-							else
-							{
-								Instance = MakeManaged<FInstance>();
-							}
-
-							if (Args.Value)
-							{
-								TManagedPtr<const FInstance> pComp = LoadInstance(Args.Value);
-
-								int32 NewComponentIndex = Instance->AddComponent();
-
-								if ( !pComp->Components.IsEmpty() )
-								{
-									Instance->Components[NewComponentIndex] = pComp->Components[0];
-
-									// Id
-									Instance->Components[NewComponentIndex].Id = Args.ExternalId;
-								}
-							}
-
-							StoreInstance(Item, Instance);
-							break;
-						}
-
 					default:
 						unimplemented();
 					}
@@ -1043,35 +705,12 @@ namespace UE::Mutable::Private
 					{
 					case 0:
 						{
-							PushOp(FScheduledOpInline(Item, 1));
 							PushOp(FScheduledOpInline(Args.instance, Item));
 							PushOp(FScheduledOpInline(Args.value, Item, false));
 
 							break;
 						}
 
-					case 1:
-						{
-							TManagedPtr<const FInstance> InInstance;
-							if (Args.instance)
-							{
-								InInstance = LoadInstance(Args.instance);
-							}
-
-							TManagedPtr<FInstance> Instance;
-							if (InInstance)
-							{
-								Instance = UE::Mutable::Private::CloneOrTakeOver<FInstance>(MoveTemp(InInstance));
-							}
-							else
-							{
-								Instance = MakeManaged<FInstance>();
-							}
-							
-							StoreInstance( Item, Instance );
-							break;
-						}
-
 					default:
 						unimplemented();
 					}
@@ -1088,64 +727,12 @@ namespace UE::Mutable::Private
 		            {
 		            case 0:
 			            {
-		            		PushOp(FScheduledOpInline(Item, 1));
 				            PushOp(FScheduledOpInline(Args.instance, Item));
 				            PushOp(FScheduledOpInline(Args.value, Item));
 
 				            break;
 			            }
 		            
-		            case 1:
-		            {
-		            	TManagedPtr<const FInstance> InInstance;
-		            	if (Args.instance)
-		            	{
-		            		InInstance = LoadInstance(Args.instance);
-		            	}
-
-		            	TManagedPtr<FInstance> Instance;
-		            	if (InInstance)
-		            	{
-		            		Instance = UE::Mutable::Private::CloneOrTakeOver<FInstance>(MoveTemp(InInstance));
-		            	}
-		            	else
-		            	{
-		            		Instance = MakeManaged<FInstance>();
-		            	}
-
-		                // Empty surfaces are ok, they still need to be created, because they may contain
-		                // additional information like internal or external IDs
-		            	TManagedPtr<const FInstance> Value;
-		            	if (Args.value)
-		            	{
-		            		Value = LoadInstance(Args.value);
-		            	}
-		            		
-	                    int32 SurfaceIndex = Instance->AddSurface(0, 0);
-
-	                    // Surface data
-	                    if (Value &&
-	                        Value->Components.Num() &&
-							Value->Components[0].LODs.Num() &&
-	                        Value->Components[0].LODs[0].Surfaces.Num())
-	                    {
-	                        Instance->Components[0].LODs[0].Surfaces[SurfaceIndex] = Value->Components[0].LODs[0].Surfaces[0];
-	                    }
-
-	                    // Name
-	                    FOperation::ADDRESS nameAd = Args.name;
-	                    check( nameAd < (uint32)Program.ConstantStrings.Num() );
-	                    const FString& Name = Program.ConstantStrings[ nameAd ];
-	                    Instance->SetSurfaceName( 0, 0, SurfaceIndex, FName(Name) );
-
-	                    // IDs
-	                    Instance->Components[0].LODs[0].Surfaces[SurfaceIndex].InternalId = Args.id;
-	                    Instance->Components[0].LODs[0].Surfaces[SurfaceIndex].ExternalId = Args.ExternalId;
-		            		
-		                StoreInstance(Item, Instance);
-		                break;
-		            }
-
 		            default:
 		                check(false);
 		            }
@@ -1167,8 +754,6 @@ namespace UE::Mutable::Private
 		            {
 		            case 0:
 		            {
-		            	PushOp(FScheduledOpInline( Item, 1));
-
 		                for (uint8 LODIndex = 0; LODIndex < LODCount; ++LODIndex)
 		                {
 							FOperation::ADDRESS LODAddress;
@@ -1185,44 +770,6 @@ namespace UE::Mutable::Private
 		                    }
 		                }
 
-
-		                break;
-		            }
-
-		            case 1:
-		            {
-		                // Assemble result
-						TManagedPtr<FInstance> Result = MakeManaged<FInstance>();
-						int32 ComponentIndex = Result->AddComponent();
-
-		                for (uint8 LODIndex = 0; LODIndex < LODCount; ++LODIndex)
-		                {
-							FOperation::ADDRESS LODAddress;
-							FMemory::Memcpy(&LODAddress, Data, sizeof(FOperation::ADDRESS));
-							Data += sizeof(FOperation::ADDRESS);
-							
-							if (LODAddress)
-		                    {						
-								bool bIsSelectedLod = ((1 << LODIndex) & LodMask ) != 0;
-
-								// Add an empty LOD even if not selected.
-								int32 InstanceLODIndex = Result->AddLOD(ComponentIndex);
-								
-								if (bIsSelectedLod)
-		                        {
-									TManagedPtr<const FInstance> LOD = LoadInstance(LODAddress);
-
-		                            // In a degenerated case, the returned pLOD may not have an LOD inside
- 									if (!LOD->Components.IsEmpty() &&
-										!LOD->Components[0].LODs.IsEmpty())
-									{
-										Result->Components[ComponentIndex].LODs[InstanceLODIndex] = LOD->Components[0].LODs[0];
-									}
-								}
-		                    }
-		                }
-
-		                StoreInstance(Item, Result);
 		                break;
 		            }
 
@@ -1242,40 +789,12 @@ namespace UE::Mutable::Private
 					{
 					case 0:
 						{
-							PushOp(FScheduledOpInline(Item, 1));
 							PushOp(FScheduledOpInline(Args.instance, Item));
 							PushOp(FScheduledOpInline(Args.value, Item, false));
 							
 							break;
 						}
 
-					case 1:
-						{
-							TManagedPtr<const FInstance> InInstance;
-							if (Args.instance)
-							{
-								InInstance = LoadInstance(Args.instance);
-							}
-
-							TManagedPtr<FInstance> Instance;
-							if (InInstance)
-							{
-								Instance = UE::Mutable::Private::CloneOrTakeOver<FInstance>(MoveTemp(InInstance));
-							}
-							else
-							{
-								Instance = MakeManaged<FInstance>();
-							}
-							
-							if (Args.value)
-							{
-        						FSkeletalMeshId SkeletalMeshId = FSystem::GetSkeletalMeshId(SkeletalMeshIdRegistry, *Model.Get(), *Params.Get(), Args.value);
-								Instance->SetSkeletalMeshId(0, SkeletalMeshId);
-							}
-							
-							StoreInstance( Item, Instance);
-							break;
-						}
 
 					default:
 						unimplemented();
@@ -1294,35 +813,12 @@ namespace UE::Mutable::Private
 					{
 					case 0:
 						{
-							PushOp(FScheduledOpInline(Item, 1));
 							PushOp(FScheduledOpInline(Args.Instance, Item));
 							PushOp(FScheduledOpInline(Args.ExtensionData, Item));
 
 							break;
 						}
 
-					case 1:
-						{
-							TManagedPtr<const FInstance> InInstance;
-							if (Args.Instance)
-							{
-								InInstance = LoadInstance(Args.Instance);
-							}
-
-							TManagedPtr<FInstance> Instance;
-							if (InInstance)
-							{
-								Instance = UE::Mutable::Private::CloneOrTakeOver<FInstance>(MoveTemp(InInstance));
-							}
-							else
-							{
-								Instance = MakeManaged<FInstance>();
-							}
-
-							StoreInstance(Item, Instance);
-							break;
-						}
-
 					default:
 						unimplemented();
 					}
@@ -1686,50 +1182,13 @@ namespace UE::Mutable::Private
 
 			case EOpType::ST_CONSTANT:
 				{
-					if (Item.bEvaluate)
-					{
-						FOperation::ResourceConstantArgs Args = Program.GetOpArgs<FOperation::ResourceConstantArgs>(Item.At);
-				
-						check(Args.value < static_cast<uint32>(Model->GetProgram().ConstantStrings.Num()));
-
-						const FString& Result = Program.ConstantStrings[Args.value];
-
-						TManagedPtr<const String> Value = MakeManaged<String>(Result);
-
-						StoreString(Item, Value);
-					}
-					else
-					{
-						StoreNone(Item);
-					}
-					
+					StoreNone(Item);
 					break;
 				}
 
 			case EOpType::ST_PARAMETER:
 				{
-					FOperation::ParameterArgs Args = Program.GetOpArgs<FOperation::ParameterArgs>(Item.At);
-
-					if (Item.bEvaluate)
-					{
-						FScheduledOp Op;
-						Op.At = Item.At;
-						Op.ExecutionIndex = Item.ExecutionIndex;
-				
-						TSharedPtr<FRangeIndex> Index = BuildCurrentOpRangeIndex(Op, *Params, *Model.Get(), *ProgramCache, Args.variable);
-
-						FString Result;
-						Params->GetStringValue(Args.variable, Result, Index.Get());
-
-						TManagedPtr<const String> Value = MakeManaged<String>(Result);
-					
-						StoreString(Item, Value);
-					}
-					else
-					{
-						StoreNone(Item);
-					}
-					
+					StoreNone(Item);
 					break;
 				}
 
@@ -2716,10 +2175,7 @@ namespace UE::Mutable::Private
 					FOperation::MaterialConstantArgs Args = Program.GetOpArgs<FOperation::MaterialConstantArgs>(Item.At);
 					PassthroughObjectLoader->Add<UMaterialInterface>(Args.ID);
 
-					TManagedPtr<const FMaterial> Material = Program.ConstantMaterials[Args.Value];
-					check(Material);
-
-					StoreMaterial(Item, Material);
+					StoreNone(Item);
 					break;
 				}
 
@@ -2733,55 +2189,14 @@ namespace UE::Mutable::Private
 					{
 						case 0:
 						{
-							// Push next stage
-							PushOp(FScheduledOpInline(Item, 1));
-
 							// Push Material Op
 							PushOp(FScheduledOpInline(Args.Material, Item));
 
 							break;
 						}
-						case 1:
-						{
-							TManagedPtr<const FMaterial> Material = LoadMaterial(Args.Material);
-							const FName& ParameterName = Program.ConstantNames[Args.ParameterName];
-							const FParameterKey ParameterKey = { ParameterName, Args.LayerIndex };
-							const EOpType InputOpType = Program.GetOpType(Item.At);
-							const EDataType InputDataType = GetOpDataType(InputOpType);
-
-							switch (InputDataType)
-							{
-								case EDataType::Image:
-								{
-									const FMaterial::FImageParameterData* ImageParameterData = Material->ImageParameters.Find(ParameterKey);
-
-									if (ImageParameterData)
-									{
-										const TVariant<FOperation::ADDRESS, TManagedPtr<const FImage>> ImageParameter = ImageParameterData->ImageParameter;
-
-										if (ImageParameter.IsType<FOperation::ADDRESS>())
-										{
-											PushOp(FScheduledOpInline(ImageParameter.Get<FOperation::ADDRESS>(), Item));
-										}
-									}
-
-									StoreNone(Item);
-									break;
-								}
-
-								case EDataType::Scalar:
-								case EDataType::Color:
-								{
-									StoreNone(Item);
-									break;
-								}
-
-								default:
-									unimplemented();
-							}
-
-							break;
-						}
+						
+						default: 
+							unimplemented();
 					}
 
 					break;
@@ -2856,7 +2271,6 @@ namespace UE::Mutable::Private
 									break;
 
 								case EDataType::Material:
-									LoadMaterial(OperantAt);
 									break;
 
 								default: 
@@ -2880,8 +2294,7 @@ namespace UE::Mutable::Private
 					const FOperation::MaterialSkeletalMeshBreakArgs Args = Program.GetOpArgs<FOperation::MaterialSkeletalMeshBreakArgs>(Item.At);
 					PushOp(FScheduledOpInline(Args.SkeletalMesh, Item));
 					
-					StoreMaterial(Item, MakeManaged<const FMaterial>());
-
+					StoreNone(Item);
 					break;
 				}
 			
@@ -2894,9 +2307,6 @@ namespace UE::Mutable::Private
 				{
 					case 0:
 					{
-						// Push next stage
-						PushOp(FScheduledOpInline(Item, 1));
-
 						//Dynamic data can't be stored into arg structures
 						const uint8* Data = Program.GetOpArgsPointer(Item.At);
 						Data += sizeof(FOperation::MaterialModifyArgs);
@@ -2937,54 +2347,10 @@ namespace UE::Mutable::Private
 
 						break;
 					}
-					case 1:
-					{
-						//Dynamic data can't be stored into arg structures
-						const uint8* Data = Program.GetOpArgsPointer(Item.At);
 						
-						//Material
-						Data += sizeof(FOperation::MaterialModifyArgs);
-
-						// Pop Material Value
-						TManagedPtr<const FMaterial> Material = LoadMaterial(Args.Material);
-
-						// Store parameter names and schedule the parameter operations
-						for (int32 ParameterIndex = 0; ParameterIndex < ParameterCount; ++ParameterIndex)
-						{
-							FOperation::ADDRESS NewName;
-							FMemory::Memcpy(&NewName, Data, sizeof(FOperation::ADDRESS));
-							Data += sizeof(FOperation::ADDRESS);
-
-							int8 LayerIndex;
-							FMemory::Memcpy(&LayerIndex, Data, sizeof(int8));
-							Data += sizeof(int8);
-
-							FOperation::ADDRESS ParameterOperation;
-							FMemory::Memcpy(&ParameterOperation, Data, sizeof(FOperation::ADDRESS));
-							Data += sizeof(FOperation::ADDRESS);
-
-							// Skip 0 addresses since do not need to load anything
-							if (ParameterOperation == 0)
-							{
-								continue;
-							}
-
-							EOpType OpType = Program.GetOpType(ParameterOperation);
-							EDataType DataType = GetOpDataType(OpType);
-
-							// Evaluate non lazy branches here
-							if (DataType == EDataType::Image)
-							{
-								// Jump Image Parameter Index
-								Data += sizeof(int32);
+				default:
+					unimplemented();
 							}
-						}
-
-						// Push material value
-						StoreMaterial(Item, Material);
-						break;
-					}
-				}
 
 				break;
 			}
@@ -2994,8 +2360,6 @@ namespace UE::Mutable::Private
 			}
 		}
 
-		check(Stack.Num() <= 1);
-
-		return Stack.Num() ? LoadInstance(Root) : MakeManaged<FInstance>();
+		check(!Stack.Num());
 	}
 }
diff --git a/Engine/Plugins/Mutable/Source/MutableRuntime/Private/MuR/CodeRunnerBegin.h b/Engine/Plugins/Mutable/Source/MutableRuntime/Private/MuR/CodeRunnerBegin.h
--- a/Engine/Plugins/Mutable/Source/MutableRuntime/Private/MuR/CodeRunnerBegin.h
+++ b/Engine/Plugins/Mutable/Source/MutableRuntime/Private/MuR/CodeRunnerBegin.h
@@ -113,39 +113,22 @@ namespace UE::Mutable::Private
 			const TSharedRef<FLiveInstance>& InLiveInstance,
 			uint32 InLodMask);
 
-		TManagedPtr<const FInstance> RunCode(const FScheduledOpInline& Root);
+		void RunCode(const FScheduledOpInline& Root);
 		
 	private:
 		FScheduledOpInline PopOp();
-		
 		void PushOp(const FScheduledOpInline& Item);
 		
-		void StoreInt(const FCacheAddressInline& To, int32 Value);
+		void StoreNone(const FCacheAddressInline& To);
 		
+		void StoreInt(const FCacheAddressInline& To, int32 Value);
 		void StoreFloat(const FCacheAddressInline& To, float Value);
-		
 		void StoreBool(const FCacheAddressInline& To, bool Value);
 		
-		void StoreInstance(const FCacheAddressInline& To, const TManagedPtr<const FInstance>& Value);
-		
-		void StoreString(const FCacheAddressInline& To, const TManagedPtr<const String>& Value);
-		
-		void StoreNone(const FCacheAddressInline& To);
-
-		void StoreMaterial(const FCacheAddressInline& To, const TManagedPtr<const FMaterial>& Value);
-
 		int32 LoadInt(const FCacheAddressInline& To);
-
 		float LoadFloat(const FCacheAddressInline& To);
-
 		bool LoadBool(const FCacheAddressInline& To);
 		
-		TManagedPtr<const FInstance> LoadInstance(const FCacheAddressInline& To);
-
-		TManagedPtr<const String> LoadString(const FCacheAddressInline& To);
-		
-		TManagedPtr<const FMaterial> LoadMaterial(const FCacheAddressInline& To);
-		
 		TSharedPtr<const FModel> Model;
 		TSharedPtr<const FParameters> Params;
 
@@ -153,26 +136,24 @@ namespace UE::Mutable::Private
 		
 		uint32 LodMask;
 
+		/** Operations we have yet to process.*/
 		TArray<FScheduledOpInline, TInlineAllocator<256, TAlignedHeapAllocator<PLATFORM_CACHE_LINE_SIZE>>> Items;
 
+		/** Stack of data objects generated during the evaluation of the mutable tree. Can be used to pass data from one stage/op to another during the tree
+		 * evaluation process
+		 */
 		FStack Stack;
 		
+		/** Collection of operations already processed. */
 		FOpSet Executed;
 
+		/** Data caches with the already resolved data. Used to prevent the reevaluation of branches that have already been evaluated.*/
 		TMap<FCacheAddressInline, float> ResultsFloat;
 		TMap<FCacheAddressInline, bool> ResultsBool;
 		TMap<FCacheAddressInline, int32> ResultsInt;
-		TMap<FCacheAddressInline, TManagedPtr<const FInstance>> ResultsInstance;
-		TMap<FCacheAddressInline, TManagedPtr<const String>> ResultsString;
-		TMap<FCacheAddressInline, TManagedPtr<const FMaterial>> ResultsMaterial;
 
 		TSharedPtr<FProgramCache> ProgramCache;
 
-		TSharedPtr<FMeshIdRegistry> MeshIdRegistry;
-		TSharedPtr<FImageIdRegistry> ImageIdRegistry;
-		TSharedPtr<FMaterialIdRegistry> MaterialIdRegistry;
-		TSharedPtr<FSkeletalMeshIdRegistry> SkeletalMeshIdRegistry;
-
 		TSharedPtr<FPassthroughObjectLoader> PassthroughObjectLoader;
 	};
 }
diff --git a/Engine/Plugins/Mutable/Source/MutableRuntime/Private/MuR/System.cpp b/Engine/Plugins/Mutable/Source/MutableRuntime/Private/MuR/System.cpp
--- a/Engine/Plugins/Mutable/Source/MutableRuntime/Private/MuR/System.cpp
+++ b/Engine/Plugins/Mutable/Source/MutableRuntime/Private/MuR/System.cpp
@@ -193,7 +193,7 @@ namespace UE::Mutable::Private
 
 
     //---------------------------------------------------------------------------------------------
-    TManagedPtr<const FInstance> FSystem::BeginUpdate_GameThread(const TSharedRef<FLiveInstance>& LiveInstance,uint32 InLodMask)
+    void FSystem::BeginUpdate_GameThread(const TSharedRef<FLiveInstance>& LiveInstance,uint32 InLodMask)
     {
 		LLM_SCOPE_BYNAME(TEXT("MutableRuntime"));
 		MUTABLE_CPUPROFILER_SCOPE(BeginUpdate_GameThread);
@@ -204,14 +204,7 @@ namespace UE::Mutable::Private
 
     	check(LiveInstance->State >= 0);
     	FScheduledOpInline Root(LiveInstance->Model->GetProgram().States[LiveInstance->State].Root, 0, 0, true);
-    	TManagedPtr<const FInstance> Result = Runner.RunCode(FScheduledOpInline(Root, 0));
-    	if (!Result)
-    	{
-    		// In case of failure return an empty instance, to prevent following code to have to check it every time
-    		Result = MakeManaged<FInstance>();
-    	}
-    	
-    	return Result;
+    	Runner.RunCode(FScheduledOpInline(Root, 0));
 	}
 
 	
