[Mutable] Allow Reroute Node to have multiple outputs.

- Deprecated Copy Mesh Section.

#Synced-CL 51705170

diff --git a/Engine/Plugins/Mutable/Source/CustomizableObjectEditor/Private/MuCOE/EdGraphSchema_CustomizableObject.cpp b/Engine/Plugins/Mutable/Source/CustomizableObjectEditor/Private/MuCOE/EdGraphSchema_CustomizableObject.cpp
--- a/Engine/Plugins/Mutable/Source/CustomizableObjectEditor/Private/MuCOE/EdGraphSchema_CustomizableObject.cpp
+++ b/Engine/Plugins/Mutable/Source/CustomizableObjectEditor/Private/MuCOE/EdGraphSchema_CustomizableObject.cpp
@@ -479,7 +479,6 @@ void UEdGraphSchema_CustomizableObject::GetGraphContextActions(FGraphContextMenu
 		{
 			ContextMenuBuilder.CreateTemplateNode<UCONodeSkeletalMeshSection>(),
 			ContextMenuBuilder.CreateTemplateNode<UCONodeSkeletalMeshSectionVariation>(),
-			ContextMenuBuilder.CreateTemplateNode<UCONodeCopySkeletalMeshSection>(),
 			ContextMenuBuilder.CreateTemplateNode<UCustomizableObjectNodeTable>(),
 			ContextMenuBuilder.CreateTemplateNode<UCustomizableObjectNodeMacroInstance>(),
 			ContextMenuBuilder.CreateTemplateNode<UCustomizableObjectNodeStaticString>()
diff --git a/Engine/Plugins/Mutable/Source/CustomizableObjectEditor/Private/MuCOE/GenerateMutableSource/GenerateMutableSource.cpp b/Engine/Plugins/Mutable/Source/CustomizableObjectEditor/Private/MuCOE/GenerateMutableSource/GenerateMutableSource.cpp
--- a/Engine/Plugins/Mutable/Source/CustomizableObjectEditor/Private/MuCOE/GenerateMutableSource/GenerateMutableSource.cpp
+++ b/Engine/Plugins/Mutable/Source/CustomizableObjectEditor/Private/MuCOE/GenerateMutableSource/GenerateMutableSource.cpp
@@ -93,7 +93,7 @@ bool FGraphCycle::FoundCycle() const
 
 
 /** Warn if the node has more outputs than it is meant to have. */
-void CheckNumOutputs(const UEdGraphPin& Pin, const FMutableGraphGenerationContext& GenerationContext)
+void CheckNode(const UEdGraphPin& Pin, const FMutableGraphGenerationContext& GenerationContext)
 {
 	if (const UCustomizableObjectNode* Typed = Cast<UCustomizableObjectNode>(Pin.GetOwningNode()))
 	{
@@ -110,9 +110,15 @@ void CheckNumOutputs(const UEdGraphPin& Pin, const FMutableGraphGenerationContex
 
 			if (numOutLinks > 1)
 			{
-				GenerationContext.CompilationContext->Log(LOCTEXT("MultipleOutgoing", "The node has several outgoing connections, but it should be limited to 1."), CastChecked<UCustomizableObjectNode>(Pin.GetOwningNode()));
+				GenerationContext.CompilationContext->Log(LOCTEXT("MultipleOutgoing", "The node has several outgoing connections, but it should be limited to 1."), Typed);
 			}
 		}
+		
+		if (Typed->IsDeprecated())
+		{
+			const FText Text = FText::Format(LOCTEXT("Deprecated", "Deprecated {0} node. Please replace it."), Typed->GetNodeTitle(ENodeTitleType::ListView));
+			GenerationContext.CompilationContext->Log(Text, Typed);
+		}
 	}
 }
 
@@ -437,7 +443,7 @@ UTexture2D* FindReferenceImage(const UEdGraphPin* Pin, FMutableGraphGenerationCo
 	check(Pin)
 	RETURN_ON_CYCLE(*Pin, GenerationContext)
 
-	CheckNumOutputs(*Pin, GenerationContext);
+	CheckNode(*Pin, GenerationContext);
 	
 	const UEdGraphSchema_CustomizableObject* Schema = GetDefault<UEdGraphSchema_CustomizableObject>();
 	const UCustomizableObjectNode* Node = CastChecked<UCustomizableObjectNode>(Pin->GetOwningNode());
@@ -873,7 +879,7 @@ UE::Mutable::Private::Ptr<UE::Mutable::Private::NodeObject> GenerateMutableSourc
 	check(Pin)
 	RETURN_ON_CYCLE(*Pin, GenerationContext)
 
-	CheckNumOutputs(*Pin, GenerationContext);
+	CheckNode(*Pin, GenerationContext);
 	
 	UCustomizableObjectNode* Node = CastChecked<UCustomizableObjectNode>(Pin->GetOwningNode());
 	
diff --git a/Engine/Plugins/Mutable/Source/CustomizableObjectEditor/Private/MuCOE/GenerateMutableSource/GenerateMutableSource.h b/Engine/Plugins/Mutable/Source/CustomizableObjectEditor/Private/MuCOE/GenerateMutableSource/GenerateMutableSource.h
--- a/Engine/Plugins/Mutable/Source/CustomizableObjectEditor/Private/MuCOE/GenerateMutableSource/GenerateMutableSource.h
+++ b/Engine/Plugins/Mutable/Source/CustomizableObjectEditor/Private/MuCOE/GenerateMutableSource/GenerateMutableSource.h
@@ -978,7 +978,7 @@ UE::Mutable::Private::Ptr<UE::Mutable::Private::NodeObject> GenerateMutableSourc
 void PopulateReferenceSkeletalMeshesData(FMutableGraphGenerationContext& GenerationContext);
 
 
-void CheckNumOutputs(const UEdGraphPin& Pin, const FMutableGraphGenerationContext& GenerationContext);
+void CheckNode(const UEdGraphPin& Pin, const FMutableGraphGenerationContext& GenerationContext);
 
 
 // TODO FutureGMT Remove generation context dependency and move to GraphTraversal.
diff --git a/Engine/Plugins/Mutable/Source/CustomizableObjectEditor/Private/MuCOE/GenerateMutableSource/GenerateMutableSourceColor.cpp b/Engine/Plugins/Mutable/Source/CustomizableObjectEditor/Private/MuCOE/GenerateMutableSource/GenerateMutableSourceColor.cpp
--- a/Engine/Plugins/Mutable/Source/CustomizableObjectEditor/Private/MuCOE/GenerateMutableSource/GenerateMutableSourceColor.cpp
+++ b/Engine/Plugins/Mutable/Source/CustomizableObjectEditor/Private/MuCOE/GenerateMutableSource/GenerateMutableSourceColor.cpp
@@ -48,7 +48,7 @@ UE::Mutable::Private::Ptr<UE::Mutable::Private::NodeColour> GenerateMutableSourc
 	check(Pin)
 	RETURN_ON_CYCLE(*Pin, GenerationContext)
 
-	CheckNumOutputs(*Pin, GenerationContext);
+	CheckNode(*Pin, GenerationContext);
 	
 	const UEdGraphSchema_CustomizableObject* Schema = GetDefault<UEdGraphSchema_CustomizableObject>();
 
diff --git a/Engine/Plugins/Mutable/Source/CustomizableObjectEditor/Private/MuCOE/GenerateMutableSource/GenerateMutableSourceComponent.cpp b/Engine/Plugins/Mutable/Source/CustomizableObjectEditor/Private/MuCOE/GenerateMutableSource/GenerateMutableSourceComponent.cpp
--- a/Engine/Plugins/Mutable/Source/CustomizableObjectEditor/Private/MuCOE/GenerateMutableSource/GenerateMutableSourceComponent.cpp
+++ b/Engine/Plugins/Mutable/Source/CustomizableObjectEditor/Private/MuCOE/GenerateMutableSource/GenerateMutableSourceComponent.cpp
@@ -42,7 +42,7 @@ UE::Mutable::Private::Ptr<UE::Mutable::Private::NodeComponent> GenerateMutableSo
 	check(Pin)
 	RETURN_ON_CYCLE(*Pin, GenerationContext)
 
-	CheckNumOutputs(*Pin, GenerationContext);
+	CheckNode(*Pin, GenerationContext);
 	
 	UCustomizableObjectNode* Node = CastChecked<UCustomizableObjectNode>(Pin->GetOwningNode());
 
diff --git a/Engine/Plugins/Mutable/Source/CustomizableObjectEditor/Private/MuCOE/GenerateMutableSource/GenerateMutableSourceExternal.cpp b/Engine/Plugins/Mutable/Source/CustomizableObjectEditor/Private/MuCOE/GenerateMutableSource/GenerateMutableSourceExternal.cpp
--- a/Engine/Plugins/Mutable/Source/CustomizableObjectEditor/Private/MuCOE/GenerateMutableSource/GenerateMutableSourceExternal.cpp
+++ b/Engine/Plugins/Mutable/Source/CustomizableObjectEditor/Private/MuCOE/GenerateMutableSource/GenerateMutableSourceExternal.cpp
@@ -33,7 +33,7 @@ UE::Mutable::Private::Ptr<UE::Mutable::Private::NodeExternal> GenerateMutableSou
 	check(Pin)
 	RETURN_ON_CYCLE(*Pin, GenerationContext)
 
-	CheckNumOutputs(*Pin, GenerationContext);
+	CheckNode(*Pin, GenerationContext);
 	
 	UCustomizableObjectNode* Node = CastChecked<UCustomizableObjectNode>(Pin->GetOwningNode());
 
diff --git a/Engine/Plugins/Mutable/Source/CustomizableObjectEditor/Private/MuCOE/GenerateMutableSource/GenerateMutableSourceFloat.cpp b/Engine/Plugins/Mutable/Source/CustomizableObjectEditor/Private/MuCOE/GenerateMutableSource/GenerateMutableSourceFloat.cpp
--- a/Engine/Plugins/Mutable/Source/CustomizableObjectEditor/Private/MuCOE/GenerateMutableSource/GenerateMutableSourceFloat.cpp
+++ b/Engine/Plugins/Mutable/Source/CustomizableObjectEditor/Private/MuCOE/GenerateMutableSource/GenerateMutableSourceFloat.cpp
@@ -43,7 +43,7 @@ UE::Mutable::Private::Ptr<UE::Mutable::Private::NodeScalar> GenerateMutableSourc
 	check(Pin)
 	RETURN_ON_CYCLE(*Pin, GenerationContext)
 
-	CheckNumOutputs(*Pin, GenerationContext);
+	CheckNode(*Pin, GenerationContext);
 	
 	const UEdGraphSchema_CustomizableObject* Schema = GetDefault<UEdGraphSchema_CustomizableObject>();
 
diff --git a/Engine/Plugins/Mutable/Source/CustomizableObjectEditor/Private/MuCOE/GenerateMutableSource/GenerateMutableSourceImage.cpp b/Engine/Plugins/Mutable/Source/CustomizableObjectEditor/Private/MuCOE/GenerateMutableSource/GenerateMutableSourceImage.cpp
--- a/Engine/Plugins/Mutable/Source/CustomizableObjectEditor/Private/MuCOE/GenerateMutableSource/GenerateMutableSourceImage.cpp
+++ b/Engine/Plugins/Mutable/Source/CustomizableObjectEditor/Private/MuCOE/GenerateMutableSource/GenerateMutableSourceImage.cpp
@@ -94,7 +94,7 @@ UE::Mutable::Private::Ptr<UE::Mutable::Private::NodeImage> GenerateMutableSource
 	check(Pin)
 	RETURN_ON_CYCLE(*Pin, GenerationContext)
 
-	CheckNumOutputs(*Pin, GenerationContext);
+	CheckNode(*Pin, GenerationContext);
 	
 	const UEdGraphSchema_CustomizableObject* Schema = GetDefault<UEdGraphSchema_CustomizableObject>();
 
@@ -1303,7 +1303,7 @@ UE::Mutable::Private::Ptr<UE::Mutable::Private::NodeImageParameter> GenerateMuta
 	check(Pin)
 	RETURN_ON_CYCLE(*Pin, GenerationContext)
 
-	CheckNumOutputs(*Pin, GenerationContext);
+	CheckNode(*Pin, GenerationContext);
 	
 	UCustomizableObjectNodeParameter* ParameterNode = CastChecked<UCustomizableObjectNodeParameter>(Pin->GetOwningNode());
 
diff --git a/Engine/Plugins/Mutable/Source/CustomizableObjectEditor/Private/MuCOE/GenerateMutableSource/GenerateMutableSourceMaterial.cpp b/Engine/Plugins/Mutable/Source/CustomizableObjectEditor/Private/MuCOE/GenerateMutableSource/GenerateMutableSourceMaterial.cpp
--- a/Engine/Plugins/Mutable/Source/CustomizableObjectEditor/Private/MuCOE/GenerateMutableSource/GenerateMutableSourceMaterial.cpp
+++ b/Engine/Plugins/Mutable/Source/CustomizableObjectEditor/Private/MuCOE/GenerateMutableSource/GenerateMutableSourceMaterial.cpp
@@ -57,7 +57,7 @@ UE::Mutable::Private::Ptr<UE::Mutable::Private::NodeMaterialParameter> GenerateM
 	check(Pin)
 	RETURN_ON_CYCLE(*Pin, GenerationContext)
 
-	CheckNumOutputs(*Pin, GenerationContext);
+	CheckNode(*Pin, GenerationContext);
 	
 	UCustomizableObjectNodeParameter* ParameterNode = CastChecked<UCustomizableObjectNodeParameter>(Pin->GetOwningNode());
 
@@ -104,7 +104,7 @@ UE::Mutable::Private::Ptr<UE::Mutable::Private::NodeMaterial> GenerateMutableSou
 	check(Pin)
 	RETURN_ON_CYCLE(*Pin, GenerationContext)
 
-	CheckNumOutputs(*Pin, GenerationContext);
+	CheckNode(*Pin, GenerationContext);
 	
 	const UEdGraphSchema_CustomizableObject* Schema = GetDefault<UEdGraphSchema_CustomizableObject>();
 
diff --git a/Engine/Plugins/Mutable/Source/CustomizableObjectEditor/Private/MuCOE/GenerateMutableSource/GenerateMutableSourceMesh.cpp b/Engine/Plugins/Mutable/Source/CustomizableObjectEditor/Private/MuCOE/GenerateMutableSource/GenerateMutableSourceMesh.cpp
--- a/Engine/Plugins/Mutable/Source/CustomizableObjectEditor/Private/MuCOE/GenerateMutableSource/GenerateMutableSourceMesh.cpp
+++ b/Engine/Plugins/Mutable/Source/CustomizableObjectEditor/Private/MuCOE/GenerateMutableSource/GenerateMutableSourceMesh.cpp
@@ -2585,7 +2585,7 @@ UE::Mutable::Private::Ptr<UE::Mutable::Private::NodeMesh> GenerateMutableSourceM
 	check(Pin)
 	RETURN_ON_CYCLE(*Pin, GenerationContext)
 
-	CheckNumOutputs(*Pin, GenerationContext);
+	CheckNode(*Pin, GenerationContext);
 
 	UCustomizableObjectNode* Node = CastChecked<UCustomizableObjectNode>(Pin->GetOwningNode());
 
diff --git a/Engine/Plugins/Mutable/Source/CustomizableObjectEditor/Private/MuCOE/GenerateMutableSource/GenerateMutableSourceModifier.cpp b/Engine/Plugins/Mutable/Source/CustomizableObjectEditor/Private/MuCOE/GenerateMutableSource/GenerateMutableSourceModifier.cpp
--- a/Engine/Plugins/Mutable/Source/CustomizableObjectEditor/Private/MuCOE/GenerateMutableSource/GenerateMutableSourceModifier.cpp
+++ b/Engine/Plugins/Mutable/Source/CustomizableObjectEditor/Private/MuCOE/GenerateMutableSource/GenerateMutableSourceModifier.cpp
@@ -54,7 +54,7 @@ UE::Mutable::Private::Ptr<UE::Mutable::Private::NodeModifier> GenerateMutableSou
 	check(Pin)
 	RETURN_ON_CYCLE(*Pin, GenerationContext)
 
-	CheckNumOutputs(*Pin, GenerationContext);
+	CheckNode(*Pin, GenerationContext);
 	
 	UCustomizableObjectNode* Node = CastChecked<UCustomizableObjectNode>(Pin->GetOwningNode());
 
diff --git a/Engine/Plugins/Mutable/Source/CustomizableObjectEditor/Private/MuCOE/GenerateMutableSource/GenerateMutableSourceProjector.cpp b/Engine/Plugins/Mutable/Source/CustomizableObjectEditor/Private/MuCOE/GenerateMutableSource/GenerateMutableSourceProjector.cpp
--- a/Engine/Plugins/Mutable/Source/CustomizableObjectEditor/Private/MuCOE/GenerateMutableSource/GenerateMutableSourceProjector.cpp
+++ b/Engine/Plugins/Mutable/Source/CustomizableObjectEditor/Private/MuCOE/GenerateMutableSource/GenerateMutableSourceProjector.cpp
@@ -16,7 +16,7 @@ UE::Mutable::Private::Ptr<UE::Mutable::Private::NodeProjector> GenerateMutableSo
 	check(Pin)
 	RETURN_ON_CYCLE(*Pin, GenerationContext)
 	
-	CheckNumOutputs(*Pin, GenerationContext);
+	CheckNode(*Pin, GenerationContext);
 	
 	UCustomizableObjectNode* Node = CastChecked<UCustomizableObjectNode>(Pin->GetOwningNode());
 
diff --git a/Engine/Plugins/Mutable/Source/CustomizableObjectEditor/Private/MuCOE/GenerateMutableSource/GenerateMutableSourceSkeletalMesh.cpp b/Engine/Plugins/Mutable/Source/CustomizableObjectEditor/Private/MuCOE/GenerateMutableSource/GenerateMutableSourceSkeletalMesh.cpp
--- a/Engine/Plugins/Mutable/Source/CustomizableObjectEditor/Private/MuCOE/GenerateMutableSource/GenerateMutableSourceSkeletalMesh.cpp
+++ b/Engine/Plugins/Mutable/Source/CustomizableObjectEditor/Private/MuCOE/GenerateMutableSource/GenerateMutableSourceSkeletalMesh.cpp
@@ -21,7 +21,7 @@ namespace UE::Mutable::Private
 		check(Pin)
 		RETURN_ON_CYCLE(*Pin, GenerationContext)
 
-		CheckNumOutputs(*Pin, GenerationContext);
+		CheckNode(*Pin, GenerationContext);
 	
 		UCustomizableObjectNodeParameter* ParameterNode = CastChecked<UCustomizableObjectNodeParameter>(Pin->GetOwningNode());
 
@@ -158,7 +158,7 @@ namespace UE::Mutable::Private
 		check(Pin)
 		RETURN_ON_CYCLE(*Pin, GenerationContext)
 
-		CheckNumOutputs(*Pin, GenerationContext);
+		CheckNode(*Pin, GenerationContext);
 
 		UCustomizableObjectNode* Node = CastChecked<UCustomizableObjectNode>(Pin->GetOwningNode());
 		
diff --git a/Engine/Plugins/Mutable/Source/CustomizableObjectEditor/Private/MuCOE/GenerateMutableSource/GenerateMutableSourceSurface.cpp b/Engine/Plugins/Mutable/Source/CustomizableObjectEditor/Private/MuCOE/GenerateMutableSource/GenerateMutableSourceSurface.cpp
--- a/Engine/Plugins/Mutable/Source/CustomizableObjectEditor/Private/MuCOE/GenerateMutableSource/GenerateMutableSourceSurface.cpp
+++ b/Engine/Plugins/Mutable/Source/CustomizableObjectEditor/Private/MuCOE/GenerateMutableSource/GenerateMutableSourceSurface.cpp
@@ -40,7 +40,7 @@ UE::Mutable::Private::Ptr<UE::Mutable::Private::NodeSurface> GenerateMutableSour
 	check(Pin)
 	RETURN_ON_CYCLE(*Pin, GenerationContext)
 
-	CheckNumOutputs(*Pin, GenerationContext);
+	CheckNode(*Pin, GenerationContext);
 
 	UCustomizableObjectNode* Node = CastChecked<UCustomizableObjectNode>(Pin->GetOwningNode());
 
diff --git a/Engine/Plugins/Mutable/Source/CustomizableObjectEditor/Private/MuCOE/GenerateMutableSource/GenerateMutableSourceTransform.cpp b/Engine/Plugins/Mutable/Source/CustomizableObjectEditor/Private/MuCOE/GenerateMutableSource/GenerateMutableSourceTransform.cpp
--- a/Engine/Plugins/Mutable/Source/CustomizableObjectEditor/Private/MuCOE/GenerateMutableSource/GenerateMutableSourceTransform.cpp
+++ b/Engine/Plugins/Mutable/Source/CustomizableObjectEditor/Private/MuCOE/GenerateMutableSource/GenerateMutableSourceTransform.cpp
@@ -17,7 +17,7 @@ UE::Mutable::Private::Ptr<UE::Mutable::Private::NodeMatrix> GenerateMutableSourc
 
 	RETURN_ON_CYCLE(*Pin, GenerationContext)
 
-	CheckNumOutputs(*Pin, GenerationContext);
+	CheckNode(*Pin, GenerationContext);
 
 	UCustomizableObjectNode* Node = CastChecked<UCustomizableObjectNode>(Pin->GetOwningNode());
 	const FGeneratedKey Key(reinterpret_cast<void*>(&GenerateMutableSourceTransform), *Pin, *Node, GenerationContext);
diff --git a/Engine/Plugins/Mutable/Source/CustomizableObjectEditor/Private/MuCOE/Nodes/CONodeCopySkeletalMeshSection.cpp b/Engine/Plugins/Mutable/Source/CustomizableObjectEditor/Private/MuCOE/Nodes/CONodeCopySkeletalMeshSection.cpp
--- a/Engine/Plugins/Mutable/Source/CustomizableObjectEditor/Private/MuCOE/Nodes/CONodeCopySkeletalMeshSection.cpp
+++ b/Engine/Plugins/Mutable/Source/CustomizableObjectEditor/Private/MuCOE/Nodes/CONodeCopySkeletalMeshSection.cpp
@@ -280,6 +280,12 @@ TArray<FString>* UCONodeCopySkeletalMeshSection::GetEnableTagsArray()
 }
 
 
+bool UCONodeCopySkeletalMeshSection::IsDeprecated() const
+{
+	return true;
+}
+
+
 UEdGraphPin* UCONodeCopySkeletalMeshSection::GetMaterialAssetPin() const
 {
 	if (UCONodeSkeletalMeshSection* NodeMaterial = GetMaterialNode())
diff --git a/Engine/Plugins/Mutable/Source/CustomizableObjectEditor/Private/MuCOE/Nodes/CONodeCopySkeletalMeshSection.h b/Engine/Plugins/Mutable/Source/CustomizableObjectEditor/Private/MuCOE/Nodes/CONodeCopySkeletalMeshSection.h
--- a/Engine/Plugins/Mutable/Source/CustomizableObjectEditor/Private/MuCOE/Nodes/CONodeCopySkeletalMeshSection.h
+++ b/Engine/Plugins/Mutable/Source/CustomizableObjectEditor/Private/MuCOE/Nodes/CONodeCopySkeletalMeshSection.h
@@ -6,8 +6,6 @@
 
 #include "CONodeCopySkeletalMeshSection.generated.h"
 
-#define UE_API CUSTOMIZABLEOBJECTEDITOR_API
-
 class UCustomizableObjectNodeRemapPins;
 class UCustomizableObjectNodeSkeletalMesh;
 class UEdGraphPin;
@@ -15,58 +13,58 @@ class UObject;
 
 
 /** Generates a new Surface with the same connections as the Parent Material Surface, but with a different Mesh. */
-UCLASS(MinimalAPI)
+UCLASS()
 class UCONodeCopySkeletalMeshSection : public UCONodeSkeletalMeshSectionBase
 {
 public:
 	GENERATED_BODY()
 
 	// UEdGraphNode interface
-	UE_API virtual FText GetNodeTitle(ENodeTitleType::Type TitleType) const override;
-	UE_API virtual FText GetTooltipText() const override;
+	virtual FText GetNodeTitle(ENodeTitleType::Type TitleType) const override;
+	virtual FText GetTooltipText() const override;
 
 	// UCustomizableObjectNode interface
-	UE_API virtual void AllocateDefaultPins(UCustomizableObjectNodeRemapPins* RemapPins) override;
-	UE_API virtual bool CanConnect(const UEdGraphPin* InOwnedInputPin, const UEdGraphPin* InOutputPin, bool& bOutIsOtherNodeBlocklisted, bool& bOutArePinsCompatible) const override;
-	UE_API virtual bool ShouldBreakExistingConnections(const UEdGraphPin* InputPin, const UEdGraphPin* OutputPin) const override;
-	UE_API virtual bool IsNodeOutDatedAndNeedsRefresh() override;
-	UE_API virtual bool ProvidesCustomPinRelevancyTest() const override;
-	UE_API virtual bool IsPinRelevant(const UEdGraphPin* Pin) const override;
-	UE_API virtual void BackwardsCompatibleFixup(int32 CustomizableObjectCustomVersion) override;
-	UE_API virtual TArray<FString> GetEnableTags(TArray<const UCustomizableObjectNodeMacroInstance*>* MacroContext = nullptr) override;
-	UE_API virtual TArray<FString>* GetEnableTagsArray() override;
+	virtual void AllocateDefaultPins(UCustomizableObjectNodeRemapPins* RemapPins) override;
+	virtual bool CanConnect(const UEdGraphPin* InOwnedInputPin, const UEdGraphPin* InOutputPin, bool& bOutIsOtherNodeBlocklisted, bool& bOutArePinsCompatible) const override;
+	virtual bool ShouldBreakExistingConnections(const UEdGraphPin* InputPin, const UEdGraphPin* OutputPin) const override;
+	virtual bool IsNodeOutDatedAndNeedsRefresh() override;
+	virtual bool ProvidesCustomPinRelevancyTest() const override;
+	virtual bool IsPinRelevant(const UEdGraphPin* Pin) const override;
+	virtual void BackwardsCompatibleFixup(int32 CustomizableObjectCustomVersion) override;
+	virtual TArray<FString> GetEnableTags(TArray<const UCustomizableObjectNodeMacroInstance*>* MacroContext = nullptr) override;
+	virtual TArray<FString>* GetEnableTagsArray() override;
+	virtual bool IsDeprecated() const override;
 
 	// UCONodeSkeletalMeshSectionBase interface
-	UE_API virtual UMaterialInterface* GetMaterial() const override;
-	UE_API virtual UEdGraphPin* GetMaterialAssetPin() const override;
-	UE_API virtual int32 GetNumParameters(EMaterialParameterType Type) const override;
-	UE_API virtual FNodeMaterialParameterId GetParameterId(EMaterialParameterType Type, int32 ParameterIndex) const override;
-	UE_API virtual FName GetParameterName(EMaterialParameterType Type, int32 ParameterIndex) const override;
-	UE_API virtual int32 GetParameterLayerIndex(EMaterialParameterType Type, int32 ParameterIndex) const override;
-	UE_API virtual FText GetParameterLayerName(EMaterialParameterType Type, int32 ParameterIndex) const override;
-	UE_API virtual bool HasParameter(const FNodeMaterialParameterId& ParameterId) const override;
-	UE_API virtual const UEdGraphPin* GetParameterPin(EMaterialParameterType Type, int32 ParameterIndex) const override;
-	UE_API virtual UEdGraphPin* GetParameterPin(const FNodeMaterialParameterId& ParameterId) const override;
-	UE_API virtual bool IsImageMutableMode(int32 ImageIndex) const override;
-	UE_API virtual bool IsImageMutableMode(const UEdGraphPin& Pin) const override;
-	UE_API virtual UTexture2D* GetImageReferenceTexture(int32 ImageIndex) const override;
-	UE_API virtual UTexture2D* GetImageValue(int32 ImageIndex) const override;
-	UE_API virtual int32 GetImageUVLayout(int32 ImageIndex) const override;
-	UE_API virtual UCONodeSkeletalMeshSection* GetMaterialNode() const override;
-	UE_API virtual UEdGraphPin* GetMeshPin() const override;
-	UE_API virtual FPostImagePinModeChangedDelegate* GetPostImagePinModeChangedDelegate() override;
-	UE_API virtual TArray<UCustomizableObjectLayout*> GetLayouts() const override;
-	UE_API virtual UEdGraphPin* OutputPin() const override;
-	UE_API virtual bool RealMaterialDataHasChanged() const override;
-	UE_API virtual UEdGraphPin* GetEnableTagsPin() const override;
+	virtual UMaterialInterface* GetMaterial() const override;
+	virtual UEdGraphPin* GetMaterialAssetPin() const override;
+	virtual int32 GetNumParameters(EMaterialParameterType Type) const override;
+	virtual FNodeMaterialParameterId GetParameterId(EMaterialParameterType Type, int32 ParameterIndex) const override;
+	virtual FName GetParameterName(EMaterialParameterType Type, int32 ParameterIndex) const override;
+	virtual int32 GetParameterLayerIndex(EMaterialParameterType Type, int32 ParameterIndex) const override;
+	virtual FText GetParameterLayerName(EMaterialParameterType Type, int32 ParameterIndex) const override;
+	virtual bool HasParameter(const FNodeMaterialParameterId& ParameterId) const override;
+	virtual const UEdGraphPin* GetParameterPin(EMaterialParameterType Type, int32 ParameterIndex) const override;
+	virtual UEdGraphPin* GetParameterPin(const FNodeMaterialParameterId& ParameterId) const override;
+	virtual bool IsImageMutableMode(int32 ImageIndex) const override;
+	virtual bool IsImageMutableMode(const UEdGraphPin& Pin) const override;
+	virtual UTexture2D* GetImageReferenceTexture(int32 ImageIndex) const override;
+	virtual UTexture2D* GetImageValue(int32 ImageIndex) const override;
+	virtual int32 GetImageUVLayout(int32 ImageIndex) const override;
+	virtual UCONodeSkeletalMeshSection* GetMaterialNode() const override;
+	virtual UEdGraphPin* GetMeshPin() const override;
+	virtual FPostImagePinModeChangedDelegate* GetPostImagePinModeChangedDelegate() override;
+	virtual TArray<UCustomizableObjectLayout*> GetLayouts() const override;
+	virtual UEdGraphPin* OutputPin() const override;
+	virtual bool RealMaterialDataHasChanged() const override;
+	virtual UEdGraphPin* GetEnableTagsPin() const override;
 	
 	// Own interface
-	UE_API UCustomizableObjectNodeSkeletalMesh* GetMeshNode() const;
+	UCustomizableObjectNodeSkeletalMesh* GetMeshNode() const;
 	
 protected:
 
-	UE_API UEdGraphPin* GetMeshSectionPin() const;
+	UEdGraphPin* GetMeshSectionPin() const;
 
 };
 
-#undef UE_API
diff --git a/Engine/Plugins/Mutable/Source/CustomizableObjectEditor/Private/MuCOE/Nodes/CustomizableObjectNodeReroute.cpp b/Engine/Plugins/Mutable/Source/CustomizableObjectEditor/Private/MuCOE/Nodes/CustomizableObjectNodeReroute.cpp
--- a/Engine/Plugins/Mutable/Source/CustomizableObjectEditor/Private/MuCOE/Nodes/CustomizableObjectNodeReroute.cpp
+++ b/Engine/Plugins/Mutable/Source/CustomizableObjectEditor/Private/MuCOE/Nodes/CustomizableObjectNodeReroute.cpp
@@ -29,7 +29,7 @@ void UCustomizableObjectNodeReroute::ReconstructNode(UCustomizableObjectNodeRema
 
 bool UCustomizableObjectNodeReroute::IsSingleOutputNode() const
 {
-	return true;
+	return false;
 }
 
 
diff --git a/Engine/Plugins/Mutable/Source/CustomizableObjectEditor/Private/MuCOE/Nodes/SCustomizableObjectNode.cpp b/Engine/Plugins/Mutable/Source/CustomizableObjectEditor/Private/MuCOE/Nodes/SCustomizableObjectNode.cpp
--- a/Engine/Plugins/Mutable/Source/CustomizableObjectEditor/Private/MuCOE/Nodes/SCustomizableObjectNode.cpp
+++ b/Engine/Plugins/Mutable/Source/CustomizableObjectEditor/Private/MuCOE/Nodes/SCustomizableObjectNode.cpp
@@ -48,6 +48,18 @@ void SCustomizableObjectNode::CreateBelowPinControls(TSharedPtr<SVerticalBox> Ma
 			];
 		}
 		
+		if (Node->IsDeprecated())
+		{
+			MainBox->AddSlot()
+			.Padding(FMargin(0.0, 2.0, 0.0, 0.0))
+			.AutoHeight()
+			[
+				SNew(SErrorText)
+				.BackgroundColor(FAppStyle::GetColor("ErrorReporting.WarningBackgroundColor"))
+				.ErrorText(LOCTEXT("SCONode_Deprecated", "DEPRECATED"))
+			];
+		}
+		
 		if (!Node->IsLoaded())
 		{
 			MainBox->AddSlot()
