[Mutable] Fix memory leak when converting Skeletal Meshes to Mutable at runtime. The memory allocation for the IO request was not freed.

#rnx

diff --git a/Engine/Plugins/Mutable/Source/CustomizableObject/Private/MuCO/UnrealConversionUtils.cpp b/Engine/Plugins/Mutable/Source/CustomizableObject/Private/MuCO/UnrealConversionUtils.cpp
--- a/Engine/Plugins/Mutable/Source/CustomizableObject/Private/MuCO/UnrealConversionUtils.cpp
+++ b/Engine/Plugins/Mutable/Source/CustomizableObject/Private/MuCO/UnrealConversionUtils.cpp
@@ -1499,9 +1499,15 @@ namespace UnrealConversionUtils
 
 					FBulkDataIORequestCallBack RequestCallback = [LoadedEvent, SkeletalMesh, WeakContext = Context.ToWeakPtr(), LODIndex](bool bWasCancelled, IBulkDataIORequest* IORequest) mutable
 						{
+							uint8* BulkData = IORequest->GetReadResults();
+
+							ON_SCOPE_EXIT
+							{
+								FMemory::Free(BulkData);
+							};
+
 							if (!bWasCancelled)
 							{
-								uint8* BulkData = IORequest->GetReadResults();
 								check(BulkData != nullptr);	
 
 								TSharedPtr<FMeshConversionContext> Context = WeakContext.Pin();
