[Mutable] Fix issues when merging Asset User Data.

#rnx
#Synced-CL 51705170

diff --git a/Engine/Plugins/Mutable/Source/CustomizableObject/Private/MuCO/CustomizableObjectInstance.cpp b/Engine/Plugins/Mutable/Source/CustomizableObject/Private/MuCO/CustomizableObjectInstance.cpp
--- a/Engine/Plugins/Mutable/Source/CustomizableObject/Private/MuCO/CustomizableObjectInstance.cpp
+++ b/Engine/Plugins/Mutable/Source/CustomizableObject/Private/MuCO/CustomizableObjectInstance.cpp
@@ -3180,6 +3180,41 @@ void InitSkeletalMeshData(const TSharedRef<FUpdateContextPrivate>& Context, USke
 	{
 		TMap<const UClass*, TArray<UAssetUserData*>> AssetUserDataByType; 
 		
+		// Hack. Get the Anim Slot data from the first LOD.
+		for (int32 LODIndex = Context->FirstLODAvailable[ComponentName]; LODIndex < Component->LODCount; ++LODIndex)
+		{
+			const TSharedRef<FInstanceUpdateData::FLOD>& LOD = Context->InstanceUpdateData.LODs[Component->FirstLOD + LODIndex];
+
+			if (!LOD->Mesh)
+			{
+				continue;
+			}
+			
+			if (LOD->Mesh->GameplayTags.Num() || LOD->Mesh->AnimationSlots.Num())
+			{
+				UCustomizableObjectInstanceUserData* AssetUserData = NewObject<UCustomizableObjectInstanceUserData>(SkeletalMesh, NAME_None, RF_Public | RF_Transactional);
+					
+				for (const FName& GameplayTag : LOD->Mesh->GameplayTags)
+				{
+					AssetUserData->AnimationGameplayTag.AddTag(FGameplayTag::RequestGameplayTag(GameplayTag, false));
+				}
+
+				for (const TPair<FName, UE::Mutable::Private::TPassthroughObjectPtr<UClass>>& Pair : LOD->Mesh->AnimationSlots)
+				{
+					FCustomizableObjectAnimationSlot Slot;
+					Slot.Name = Pair.Key;
+					Slot.AnimInstance = Pair.Value.Get();
+
+					AssetUserData->AnimationSlots.Add(Slot);
+				}	
+				
+				TArray<UAssetUserData*>& Result = AssetUserDataByType.FindOrAdd(AssetUserData->GetClass(), {});
+				Result.AddUnique(AssetUserData);
+			}
+
+			break;
+		}
+		
 		// Hack. Get the asset user data from the first LOD.
 		for (int32 LODIndex = Context->FirstLODAvailable[ComponentName]; LODIndex < Component->LODCount; ++LODIndex)
 		{
@@ -3229,43 +3264,6 @@ void InitSkeletalMeshData(const TSharedRef<FUpdateContextPrivate>& Context, USke
 		}
 	}
 	
-	// Anim Slots
-	{
-		// Hack. Get the Anim Slot data from the first LOD.
-		for (int32 LODIndex = Context->FirstLODAvailable[ComponentName]; LODIndex < Component->LODCount; ++LODIndex)
-		{
-			const TSharedRef<FInstanceUpdateData::FLOD>& LOD = Context->InstanceUpdateData.LODs[Component->FirstLOD + LODIndex];
-
-			if (!LOD->Mesh)
-			{
-				continue;
-			}
-			
-			if (LOD->Mesh->GameplayTags.Num() || LOD->Mesh->AnimationSlots.Num())
-			{
-				UCustomizableObjectInstanceUserData* UserData = NewObject<UCustomizableObjectInstanceUserData>(SkeletalMesh, NAME_None, RF_Public | RF_Transactional);
-					
-				for (const FName& GameplayTag : LOD->Mesh->GameplayTags)
-				{
-					UserData->AnimationGameplayTag.AddTag(FGameplayTag::RequestGameplayTag(GameplayTag, false));
-				}
-
-				for (const TPair<FName, UE::Mutable::Private::TPassthroughObjectPtr<UClass>>& Pair : LOD->Mesh->AnimationSlots)
-				{
-					FCustomizableObjectAnimationSlot Slot;
-					Slot.Name = Pair.Key;
-					Slot.AnimInstance = Pair.Value.Get();
-
-					UserData->AnimationSlots.Add(Slot);
-				}	
-				
-				SkeletalMesh->AddAssetUserData(UserData);
-			}
-
-			break;
-		}
-	}
-	
 	// Allocate resources for rendering and add LOD Info
 	{
 		MUTABLE_CPUPROFILER_SCOPE(InitSkeletalMesh_AddLODData);
diff --git a/Engine/Plugins/Mutable/Source/CustomizableObject/Private/MuCO/CustomizableObjectInstanceAssetUserData.cpp b/Engine/Plugins/Mutable/Source/CustomizableObject/Private/MuCO/CustomizableObjectInstanceAssetUserData.cpp
--- a/Engine/Plugins/Mutable/Source/CustomizableObject/Private/MuCO/CustomizableObjectInstanceAssetUserData.cpp
+++ b/Engine/Plugins/Mutable/Source/CustomizableObject/Private/MuCO/CustomizableObjectInstanceAssetUserData.cpp
@@ -2,6 +2,8 @@
 
 #include "MuCO/CustomizableObjectInstanceAssetUserData.h"
 
+#include "MuCO/CustomizableObject.h"
+
 #if WITH_EDITOR
 #include "Logging/MessageLog.h"
 #endif
@@ -18,3 +20,34 @@ void UCustomizableObjectInstanceUserData::SetAnimationGameplayTags(const FGamepl
 	AnimationGameplayTag = InstanceTags;
 }
 
+
+bool UCustomizableObjectInstanceUserData::Merge(const TNotNull<UAssetUserData*> Other)
+{
+	UCustomizableObjectInstanceUserData* TypedOther = Cast<UCustomizableObjectInstanceUserData>(Other);
+	if (!TypedOther)
+	{
+		return false;
+	}
+
+	AnimationGameplayTag.AppendTags(TypedOther->AnimationGameplayTag);
+    
+	for (const FCustomizableObjectAnimationSlot& AnimationSlot : TypedOther->AnimationSlots)
+	{
+		FCustomizableObjectAnimationSlot* FindResult = AnimationSlots.FindByPredicate([&](const FCustomizableObjectAnimationSlot& Slot)
+		{
+			return Slot.Name == AnimationSlot.Name;
+		});
+		
+		if (!FindResult)
+		{
+			AnimationSlots.Add(AnimationSlot);
+		}
+		else if (FindResult->AnimInstance != AnimationSlot.AnimInstance)
+		{
+			UE_LOGF(LogMutable, Warning, "Unable to merge AnimBP Slots. AnimBP Slot Name [%ls] already exists", *AnimationSlot.Name.ToString());
+		}
+	}
+		
+	return true;
+}
+
diff --git a/Engine/Plugins/Mutable/Source/CustomizableObject/Public/MuCO/CustomizableObjectInstanceAssetUserData.h b/Engine/Plugins/Mutable/Source/CustomizableObject/Public/MuCO/CustomizableObjectInstanceAssetUserData.h
--- a/Engine/Plugins/Mutable/Source/CustomizableObject/Public/MuCO/CustomizableObjectInstanceAssetUserData.h
+++ b/Engine/Plugins/Mutable/Source/CustomizableObject/Public/MuCO/CustomizableObjectInstanceAssetUserData.h
@@ -38,7 +38,9 @@ public:
 	/** Sets the list of tags for this instance. */
 	UFUNCTION(BlueprintCallable, Category = CustomizableObjectInstance)
 	CUSTOMIZABLEOBJECT_API void SetAnimationGameplayTags(const FGameplayTagContainer& InstanceTags);
-	
+
+	CUSTOMIZABLEOBJECT_API virtual bool Merge(const TNotNull<UAssetUserData*> Other) override;
+
 	UPROPERTY(VisibleAnywhere, BlueprintReadWrite, Category = CustomizableObjectInstance)
 	FGameplayTagContainer AnimationGameplayTag;
 
