[Mutable] Change all instances of UE_LOG to UE_LOGF in MutableTools and MutableRuntime modules.

#rnx
#Synced-CL 51716106

diff --git a/Engine/Plugins/Mutable/Source/MutableRuntime/Private/MuR/BlockCompression/Miro/Miro.cpp b/Engine/Plugins/Mutable/Source/MutableRuntime/Private/MuR/BlockCompression/Miro/Miro.cpp
--- a/Engine/Plugins/Mutable/Source/MutableRuntime/Private/MuR/BlockCompression/Miro/Miro.cpp
+++ b/Engine/Plugins/Mutable/Source/MutableRuntime/Private/MuR/BlockCompression/Miro/Miro.cpp
@@ -5074,14 +5074,14 @@ namespace impl
 #ifdef UE_MIRO_DEBUG
 			if (s_DebugBlock == s_CurrentBlock)
 			{
-				UE_LOG(LogMutableCore, Log, TEXT("endpoints f : (%.3f,%.3f,%.3f,%.3f) (%.3f,%.3f,%.3f,%.3f)"),
+				UE_LOGF(LogMutableCore, Log, "endpoints f : (%.3f,%.3f,%.3f,%.3f) (%.3f,%.3f,%.3f,%.3f)",
 					e0.r, e0.g, e0.b, e0.a,
 					e1.r, e1.g, e1.b, e1.a);
-				UE_LOG(LogMutableCore, Log, TEXT("endpoints d : (%3d,%3d,%3d,%3d) (%3d,%3d,%3d,%3d)"),
+				UE_LOGF(LogMutableCore, Log, "endpoints d : (%3d,%3d,%3d,%3d) (%3d,%3d,%3d,%3d)",
 					Option.endpoint_unquantized[0].r, Option.endpoint_unquantized[0].g, Option.endpoint_unquantized[0].b, Option.endpoint_unquantized[0].a,
 					Option.endpoint_unquantized[1].r, Option.endpoint_unquantized[1].g, Option.endpoint_unquantized[1].b, Option.endpoint_unquantized[1].a);
 
-				UE_LOG(LogMutableCore, Log, TEXT("texel block pixels :"));
+				UE_LOGF(LogMutableCore, Log, "texel block pixels :");
 				int T = 0;
 				for (uint64 Y = 0; Y < BLOCK_Y; ++Y)
 				{
@@ -5092,10 +5092,10 @@ namespace impl
 						Line += FString::Printf(TEXT("(%3d,%3d,%3d,%3d), "), texel.r, texel.g, texel.b, texel.a);
 					}
 
-					UE_LOG(LogMutableCore, Log, TEXT("%s"), *Line);
+					UE_LOGF(LogMutableCore, Log, "%ls", *Line);
 				}
 
-				UE_LOG(LogMutableCore, Log, TEXT("texel grid pixels :"));
+				UE_LOGF(LogMutableCore, Log, "texel grid pixels :");
 				T = 0;
 				for (uint64 Y = 0; Y < GRID_Y; ++Y)
 				{
@@ -5106,10 +5106,10 @@ namespace impl
 						Line += FString::Printf(TEXT("(%3d,%3d,%3d,%3d), "), texel.r, texel.g, texel.b, texel.a);
 					}
 
-					UE_LOG(LogMutableCore, Log, TEXT("%s"), *Line);
+					UE_LOGF(LogMutableCore, Log, "%ls", *Line);
 				}
 
-				UE_LOG(LogMutableCore, Log, TEXT("texel grid encoded :"));
+				UE_LOGF(LogMutableCore, Log, "texel grid encoded :");
 				T = 0;
 				for (uint64 Y = 0; Y < GRID_Y; ++Y)
 				{
@@ -5120,7 +5120,7 @@ namespace impl
 						Line += FString::Printf(TEXT("%3d, "), texel);
 					}
 
-					UE_LOG(LogMutableCore, Log, TEXT("%s"), *Line);
+					UE_LOGF(LogMutableCore, Log, "%ls", *Line);
 				}
 			}
 #endif
@@ -7139,7 +7139,7 @@ namespace impl
   ((byte) & 0x002 ? L'1' : L'0'), \
   ((byte) & 0x001 ? L'1' : L'0') 
 
-			UE_LOG(LogMutableCore, Log, TEXT("block_mode : %c%c %c%c%c%c %c%c%c%c%c"), MODE_TO_BINARY(BlockMode));
+			UE_LOGF(LogMutableCore, Log, "block_mode : %c%c %c%c%c%c %c%c%c%c%c", MODE_TO_BINARY(BlockMode));
 
 
 			// Check for block mode errors.
@@ -7158,11 +7158,11 @@ namespace impl
 			const int numWeightDataBits = computeNumRequiredBits(blockMode.weightISEParams, numWeights);
 			const int numPartitions = (int)blockData.getBits(11, 12) + 1;
 
-			UE_LOG(LogMutableCore, Log, TEXT("weight_grid : %d x %d"), blockMode.weightGridWidth, blockMode.weightGridHeight);
-			UE_LOG(LogMutableCore, Log, TEXT("num_weights : %d"), numWeights);
-			UE_LOG(LogMutableCore, Log, TEXT("num_weights_data_bits : %d"), numWeightDataBits);
-			UE_LOG(LogMutableCore, Log, TEXT("num_partitions : %d"), numPartitions);
-			UE_LOG(LogMutableCore, Log, TEXT("dual_plane : %d"), blockMode.isDualPlane);
+			UE_LOGF(LogMutableCore, Log, "weight_grid : %d x %d", blockMode.weightGridWidth, blockMode.weightGridHeight);
+			UE_LOGF(LogMutableCore, Log, "num_weights : %d", numWeights);
+			UE_LOGF(LogMutableCore, Log, "num_weights_data_bits : %d", numWeightDataBits);
+			UE_LOGF(LogMutableCore, Log, "num_partitions : %d", numPartitions);
+			UE_LOGF(LogMutableCore, Log, "dual_plane : %d", blockMode.isDualPlane);
 
 			// Check for errors in weight grid, partition and dual-plane parameters.
 			if (numWeights > 64 ||
@@ -7191,9 +7191,9 @@ namespace impl
 			decodeColorEndpointModes(&colorEndpointModes[0], blockData, numPartitions, extraCemBitsStart);
 			const int numColorEndpointValues = computeNumColorEndpointValues(colorEndpointModes, numPartitions);
 
-			UE_LOG(LogMutableCore, Log, TEXT("num_endpoint_data_bits : %d"), numBitsForColorEndpoints);
-			UE_LOG(LogMutableCore, Log, TEXT("endpoint_mode : %d"), colorEndpointModes[0]);
-			UE_LOG(LogMutableCore, Log, TEXT("num_endpoint_values : %d"), numColorEndpointValues);
+			UE_LOGF(LogMutableCore, Log, "num_endpoint_data_bits : %d", numBitsForColorEndpoints);
+			UE_LOGF(LogMutableCore, Log, "endpoint_mode : %d", colorEndpointModes[0]);
+			UE_LOGF(LogMutableCore, Log, "num_endpoint_values : %d", numColorEndpointValues);
 
 			// Check for errors in color endpoint value count.
 			if (numColorEndpointValues > 18 || numBitsForColorEndpoints < (int)deDivRoundUp32(13 * numColorEndpointValues, 5))
@@ -7205,7 +7205,7 @@ namespace impl
 			computeColorEndpoints(&colorEndpoints[0], blockData, &colorEndpointModes[0], numPartitions, numColorEndpointValues,
 				computeMaximumRangeISEParams(numBitsForColorEndpoints, numColorEndpointValues), numBitsForColorEndpoints);
 
-			UE_LOG(LogMutableCore, Log, TEXT("endpoints : ( %3d, %3d, %3d, %3d ), ( %3d, %3d, %3d, %3d )"),
+			UE_LOGF(LogMutableCore, Log, "endpoints : ( %3d, %3d, %3d, %3d ), ( %3d, %3d, %3d, %3d )",
 				colorEndpoints[0].e0[0], colorEndpoints[0].e0[1], colorEndpoints[0].e0[2], colorEndpoints[0].e0[3],
 				colorEndpoints[0].e1[0], colorEndpoints[0].e1[1], colorEndpoints[0].e1[2], colorEndpoints[0].e1[3]);
 
@@ -7227,13 +7227,13 @@ namespace impl
 			int i = 0;
 			//for (int32 Row = 0; Row < blockHeight; ++Row)
 			//{
-			//	UE_LOG(LogMutableCore, Log, TEXT("texel weights : %3d %3d %3d %3d %3d %3d %3d %3d"), texelWeights[i++].w[0], texelWeights[i++].w[0], texelWeights[i++].w[0], texelWeights[i++].w[0], texelWeights[i++].w[0], texelWeights[i++].w[0], texelWeights[i++].w[0], texelWeights[i++].w[0]);
+			//	UE_LOGF(LogMutableCore, Log, "texel weights : %3d %3d %3d %3d %3d %3d %3d %3d", texelWeights[i++].w[0], texelWeights[i++].w[0], texelWeights[i++].w[0], texelWeights[i++].w[0], texelWeights[i++].w[0], texelWeights[i++].w[0], texelWeights[i++].w[0], texelWeights[i++].w[0]);
 			//}
 
 			i = 0;
 			for (int32 Row = 0; Row < blockMode.weightGridHeight; ++Row)
 			{
-				UE_LOG(LogMutableCore, Log, TEXT("texel weights : %3d %3d %3d %3d %3d %3d %3d"), unquantizedWeights[i++], unquantizedWeights[i++], unquantizedWeights[i++], unquantizedWeights[i++], unquantizedWeights[i++], unquantizedWeights[i++], unquantizedWeights[i++]);
+				UE_LOGF(LogMutableCore, Log, "texel weights : %3d %3d %3d %3d %3d %3d %3d", unquantizedWeights[i++], unquantizedWeights[i++], unquantizedWeights[i++], unquantizedWeights[i++], unquantizedWeights[i++], unquantizedWeights[i++], unquantizedWeights[i++]);
 			}
 
 			// Set texel colors.
@@ -7243,7 +7243,7 @@ namespace impl
 			DecompressResult Result = setTexelColorsSRGB<BlockSize>(dst, &colorEndpoints[0], &texelWeights[0], ccs, partitionIndexSeed, numPartitions, &colorEndpointModes[0]);
 			check(Result == DecompressResult::DECOMPRESS_RESULT_VALID_BLOCK);
 
-			UE_LOG(LogMutableCore, Log, TEXT("decoded texels:"));
+			UE_LOGF(LogMutableCore, Log, "decoded texels:");
 			i = 0;
 			for (int32 Row = 0; Row < BlockSize; ++Row)
 			{
@@ -7259,7 +7259,7 @@ namespace impl
 					i++;
 				}
 
-				UE_LOG(LogMutableCore, Log, TEXT("%s"), *Line);
+				UE_LOGF(LogMutableCore, Log, "%ls", *Line);
 			}
 
 #undef MODE_TO_BINARY
diff --git a/Engine/Plugins/Mutable/Source/MutableRuntime/Private/MuR/CodeRunner.cpp b/Engine/Plugins/Mutable/Source/MutableRuntime/Private/MuR/CodeRunner.cpp
--- a/Engine/Plugins/Mutable/Source/MutableRuntime/Private/MuR/CodeRunner.cpp
+++ b/Engine/Plugins/Mutable/Source/MutableRuntime/Private/MuR/CodeRunner.cpp
@@ -4939,7 +4939,7 @@ namespace UE::Mutable::Private
 				else
 				{
 					// This happens very often when skipping mips, and floods the log.
-					//UE_LOG( LogMutableCore, Verbose, TEXT("Skipped patch operation for image not fitting the block compression size. Small image? Patch rect is (%d, %d), (%d, %d), base is (%d, %d)"),
+					//UE_LOGF( LogMutableCore, Verbose, "Skipped patch operation for image not fitting the block compression size. Small image? Patch rect is (%d, %d), (%d, %d), base is (%d, %d)",
 					//	rect.min[0], rect.min[1], rect.size[0], rect.size[1], Result->GetSizeX(), Result->GetSizeY());
 				}
 
@@ -7079,11 +7079,11 @@ namespace UE::Mutable::Private
     //---------------------------------------------------------------------------------------------
     void CodeRunner::RunCode(const FScheduledOp& Item, uint32 LodMask)
     {
-		//UE_LOG( LogMutableCore, Log, TEXT("Running :%5d , %d "), Item.At, Item.Stage );
+		//UE_LOGF( LogMutableCore, Log, "Running :%5d , %d ", Item.At, Item.Stage );
 		check( Item.Type == FScheduledOp::EType::Full );
 
 		EOpType type = Program.GetOpType(Item.At);
-		//UE_LOG(LogMutableCore, Log, TEXT("Running :%5d , %d, of type %d "), Item.At, Item.Stage, type);
+		//UE_LOGF(LogMutableCore, Log, "Running :%5d , %d, of type %d ", Item.At, Item.Stage, type);
 
 		// Very spammy, for debugging purposes.
 		//if (System)
diff --git a/Engine/Plugins/Mutable/Source/MutableRuntime/Private/MuR/CodeRunnerTasked.cpp b/Engine/Plugins/Mutable/Source/MutableRuntime/Private/MuR/CodeRunnerTasked.cpp
--- a/Engine/Plugins/Mutable/Source/MutableRuntime/Private/MuR/CodeRunnerTasked.cpp
+++ b/Engine/Plugins/Mutable/Source/MutableRuntime/Private/MuR/CodeRunnerTasked.cpp
@@ -231,7 +231,7 @@ namespace UE::Mutable::Private
 			//		}
 			//	}
 
-			//	UE_LOG(LogMutableCore, Log, TEXT("Tasks: %5d open, %5d issued, %5d closed, %d closed ready"), OpenTasks.Num(), IssuedTasks.Num(), ClosedTasks.Num(), ClosedReady);
+			//	UE_LOGF(LogMutableCore, Log, "Tasks: %5d open, %5d issued, %5d closed, %d closed ready", OpenTasks.Num(), IssuedTasks.Num(), ClosedTasks.Num(), ClosedReady);
 			//}
 
 			for (int32 Index = 0; bSuccess && Index < IssuedTasks.Num(); )
@@ -460,8 +460,8 @@ namespace UE::Mutable::Private
 		
 		if (ProfileContext)
 		{
-			UE_LOG(LogMutableCore, Log, TEXT("Mutable Heap Bytes: %d"), HeapData.Num()* HeapData.GetTypeSize());
-			UE_LOG(LogMutableCore, Log, TEXT("Ran ops : %5d "), ProfileContext->NumRunOps);
+			UE_LOGF(LogMutableCore, Log, "Mutable Heap Bytes: %d", HeapData.Num()* HeapData.GetTypeSize());
+			UE_LOGF(LogMutableCore, Log, "Ran ops : %5d ", ProfileContext->NumRunOps);
 
 			constexpr int32 HistogramSize = 8;
 			int32 MostCommonOps[HistogramSize] = {};
@@ -486,7 +486,7 @@ namespace UE::Mutable::Private
 
 			for (int32 HistIndex = 0; HistIndex < HistogramSize; ++HistIndex)
 			{
-				UE_LOG(LogMutableCore, Log, TEXT("    op %4d, %4d times."), MostCommonOps[HistIndex], ProfileContext->RunOpsPerType[MostCommonOps[HistIndex]]);
+				UE_LOGF(LogMutableCore, Log, "    op %4d, %4d times.", MostCommonOps[HistIndex], ProfileContext->RunOpsPerType[MostCommonOps[HistIndex]]);
 			}
 		}
     }
@@ -1820,7 +1820,7 @@ namespace UE::Mutable::Private
 			{
 				if (FMath::FRand() < RomLoadFailureRate)
 				{
-					UE_LOG(LogMutableCore, Warning, TEXT("Mesh rom was forced to faile with CVar."));
+					UE_LOGF(LogMutableCore, Warning, "Mesh rom was forced to faile with CVar.");
 					bOutFailed = true;
 					return false;
 				}
@@ -1960,7 +1960,7 @@ namespace UE::Mutable::Private
 			{
 				if (FMath::FRand() < RomLoadFailureRate)
 				{
-					UE_LOG(LogMutableCore, Warning, TEXT("Image rom was forced to fail with CVar."));
+					UE_LOGF(LogMutableCore, Warning, "Image rom was forced to fail with CVar.");
 					bOutFailed = true;
 					return false;
 				}
@@ -2248,13 +2248,13 @@ namespace UE::Mutable::Private
 				Issued = MakeShared<FLoadImageRomsTask>(Item, ImageIndexBegin, ImageIndexEnd);
 
 				if (DebugRom && (DebugRomAll || ImageIndex == DebugImageIndex))
-					UE_LOG(LogMutableCore, Log, TEXT("Issuing image %d skipping %d ."), ImageIndex, ReallySkip);
+					UE_LOGF(LogMutableCore, Log, "Issuing image %d skipping %d .", ImageIndex, ReallySkip);
 			}
 			else
 			{
 				// If already available, the rest of the constant code will run right away.
 				if (DebugRom && (DebugRomAll || ImageIndex == DebugImageIndex))
-					UE_LOG(LogMutableCore, Log, TEXT("Image %d skipping %d is already loaded."), ImageIndex, ReallySkip);
+					UE_LOGF(LogMutableCore, Log, "Image %d skipping %d is already loaded.", ImageIndex, ReallySkip);
 			}
 			break;
 		}		
diff --git a/Engine/Plugins/Mutable/Source/MutableRuntime/Private/MuR/MemoryTrackingUtils.cpp b/Engine/Plugins/Mutable/Source/MutableRuntime/Private/MuR/MemoryTrackingUtils.cpp
--- a/Engine/Plugins/Mutable/Source/MutableRuntime/Private/MuR/MemoryTrackingUtils.cpp
+++ b/Engine/Plugins/Mutable/Source/MutableRuntime/Private/MuR/MemoryTrackingUtils.cpp
@@ -191,13 +191,13 @@ namespace UE::Mutable::Private
 			{
 				if (!UE_MUTABLE_TRACK_ALLOCATOR_MEMORY_PEAK)
 				{
-					UE_LOG(LogMutableCore, Display, TEXT("The build does not have UE_MUTABLE_TRACK_ALLOCATOR_MEMORY_PEAK enabled, no counting performed."));
+					UE_LOGF(LogMutableCore, Display, "The build does not have UE_MUTABLE_TRACK_ALLOCATOR_MEMORY_PEAK enabled, no counting performed.");
 					return;
 				}
 
 				if (!FMemoryTrackingConsoleFlags::bEnableGlobalMemoryTracking)
 				{
-					UE_LOG(LogMutableCore, Display, TEXT("mutable.EnableGlobalMemoryTracking is set to false, no counting performed."));
+					UE_LOGF(LogMutableCore, Display, "mutable.EnableGlobalMemoryTracking is set to false, no counting performed.");
 					return;
 				}
 
@@ -217,13 +217,13 @@ namespace UE::Mutable::Private
 			{
 				if (!UE_MUTABLE_TRACK_ALLOCATOR_MEMORY_PEAK)
 				{
-					UE_LOG(LogMutableCore, Display, TEXT("The build does not have UE_MUTABLE_TRACK_ALLOCATOR_MEMORY_PEAK enabled, no counting performed."));
+					UE_LOGF(LogMutableCore, Display, "The build does not have UE_MUTABLE_TRACK_ALLOCATOR_MEMORY_PEAK enabled, no counting performed.");
 					return;
 				}
 
 				if (!FMemoryTrackingConsoleFlags::bEnableGlobalMemoryTracking)
 				{
-					UE_LOG(LogMutableCore, Display, TEXT("mutable.EnableGlobalMemoryTracking is set to false, no counting performed."));
+					UE_LOGF(LogMutableCore, Display, "mutable.EnableGlobalMemoryTracking is set to false, no counting performed.");
 					return;
 				}
 
@@ -244,13 +244,13 @@ namespace UE::Mutable::Private
 			{
 				if (!UE_MUTABLE_TRACK_ALLOCATOR_MEMORY_PEAK)
 				{
-					UE_LOG(LogMutableCore, Display, TEXT("The build does not have UE_MUTABLE_TRACK_ALLOCATOR_MEMORY_PEAK enabled, no counting performed."));
+					UE_LOGF(LogMutableCore, Display, "The build does not have UE_MUTABLE_TRACK_ALLOCATOR_MEMORY_PEAK enabled, no counting performed.");
 					return;
 				}
 
 				if (!FMemoryTrackingConsoleFlags::bEnableGlobalMemoryTracking)
 				{
-					UE_LOG(LogMutableCore, Display, TEXT("mutable.EnableGlobalMemoryTracking is set to false, no counting performed."));
+					UE_LOGF(LogMutableCore, Display, "mutable.EnableGlobalMemoryTracking is set to false, no counting performed.");
 					return;
 				}
 
@@ -265,13 +265,13 @@ namespace UE::Mutable::Private
 			{
 				if (!UE_MUTABLE_TRACK_ALLOCATOR_MEMORY_PEAK)
 				{
-					UE_LOG(LogMutableCore, Display, TEXT("The build does not have UE_MUTABLE_TRACK_ALLOCATOR_MEMORY_PEAK enabled, no counting performed."));
+					UE_LOGF(LogMutableCore, Display, "The build does not have UE_MUTABLE_TRACK_ALLOCATOR_MEMORY_PEAK enabled, no counting performed.");
 					return;
 				}
 
 				if (!FMemoryTrackingConsoleFlags::bEnableGlobalMemoryTracking)
 				{
-					UE_LOG(LogMutableCore, Display, TEXT("mutable.EnableGlobalMemoryTracking is set to false, no counting performed."));
+					UE_LOGF(LogMutableCore, Display, "mutable.EnableGlobalMemoryTracking is set to false, no counting performed.");
 					return;
 				}
 
diff --git a/Engine/Plugins/Mutable/Source/MutableRuntime/Private/MuR/Model.cpp b/Engine/Plugins/Mutable/Source/MutableRuntime/Private/MuR/Model.cpp
--- a/Engine/Plugins/Mutable/Source/MutableRuntime/Private/MuR/Model.cpp
+++ b/Engine/Plugins/Mutable/Source/MutableRuntime/Private/MuR/Model.cpp
@@ -211,11 +211,11 @@ namespace UE::Mutable::Private
             return a.first>b.first;
         });
 
-        UE_LOG(LogMutableCore,Log, TEXT("Op histogram (%llu ops):"), OpAddress.Num());
+        UE_LOGF(LogMutableCore,Log, "Op histogram (%llu ops):", OpAddress.Num());
         for(int32 i=0; i<8; ++i)
         {
             float p = sorted[i].first/float(OpAddress.Num())*100.0f;
-            UE_LOG(LogMutableCore,Log, TEXT("  %3.2f%% : %d"), p, (int32)sorted[i].second );
+            UE_LOGF(LogMutableCore,Log, "  %3.2f%% : %d", p, (int32)sorted[i].second );
         }
 #endif
     }
diff --git a/Engine/Plugins/Mutable/Source/MutableRuntime/Private/MuR/OpImagePixelFormat.cpp b/Engine/Plugins/Mutable/Source/MutableRuntime/Private/MuR/OpImagePixelFormat.cpp
--- a/Engine/Plugins/Mutable/Source/MutableRuntime/Private/MuR/OpImagePixelFormat.cpp
+++ b/Engine/Plugins/Mutable/Source/MutableRuntime/Private/MuR/OpImagePixelFormat.cpp
@@ -1112,7 +1112,7 @@ void FImageOperator::ImagePixelFormat(bool& bOutSuccess, int32 CompressionQualit
 			}
 			default:
 			{
-				UE_LOG(LogMutableCore, Log, TEXT("Image format conversion not implemented. Expensive generic one used. "));
+				UE_LOGF(LogMutableCore, Log, "Image format conversion not implemented. Expensive generic one used. ");
 
 				// Case not implemented, try a generic more expensive operation.
 				TManagedPtr<FImage> TempBase = ImagePixelFormat(CompressionQuality, Base, EImageFormat::RGBA_UByte);
@@ -1137,7 +1137,7 @@ void FImageOperator::ImagePixelFormat(bool& bOutSuccess, int32 CompressionQualit
 			}
 			default:
 			{
-				UE_LOG(LogMutableCore, Log, TEXT("Image format conversion not implemented. Expensive generic one used."));
+				UE_LOGF(LogMutableCore, Log, "Image format conversion not implemented. Expensive generic one used.");
 
 				// Case not implemented, try a generic more expensive operatino.
 				TManagedPtr<FImage> TempBase = ImagePixelFormat(CompressionQuality, Base, EImageFormat::RGBA_UByte);
@@ -1175,7 +1175,7 @@ void FImageOperator::ImagePixelFormat(bool& bOutSuccess, int32 CompressionQualit
 			}
 			default:
 			{
-				UE_LOG(LogMutableCore, Log, TEXT("Image format conversion not implemented. Expensive generic one used."));
+				UE_LOGF(LogMutableCore, Log, "Image format conversion not implemented. Expensive generic one used.");
 
 				// Case not implemented, try a generic more expensive operation.
 				TManagedPtr<FImage> TempBase = ImagePixelFormat(CompressionQuality, Base, EImageFormat::RGBA_UByte);
@@ -1231,7 +1231,7 @@ void FImageOperator::ImagePixelFormat(bool& bOutSuccess, int32 CompressionQualit
 			}
 			default:
 			{
-				UE_LOG(LogMutableCore, Log, TEXT("Image format conversion not implemented. Expensive generic one used. "));
+				UE_LOGF(LogMutableCore, Log, "Image format conversion not implemented. Expensive generic one used. ");
 
 				TManagedPtr<FImage> TempBase = ImagePixelFormat(CompressionQuality, Base, EImageFormat::RGBA_UByte, -1);
 				ImagePixelFormatInternal::BlockCompressionLODRangeHelper(
@@ -1273,7 +1273,7 @@ void FImageOperator::ImagePixelFormat(bool& bOutSuccess, int32 CompressionQualit
 			}
 			default:
 			{
-				UE_LOG(LogMutableCore, Log, TEXT("Image format conversion not implemented. Expensive generic one used. "));
+				UE_LOGF(LogMutableCore, Log, "Image format conversion not implemented. Expensive generic one used. ");
 
 				TManagedPtr<FImage> TempBase = ImagePixelFormat(CompressionQuality, Base, EImageFormat::RGBA_UByte, -1);
 				ImagePixelFormatInternal::BlockCompressionLODRangeHelper(
@@ -1309,7 +1309,7 @@ void FImageOperator::ImagePixelFormat(bool& bOutSuccess, int32 CompressionQualit
 			}
 			default:
 			{
-				UE_LOG(LogMutableCore, Log, TEXT("Image format conversion not implemented. Expensive generic one used. "));
+				UE_LOGF(LogMutableCore, Log, "Image format conversion not implemented. Expensive generic one used. ");
 
 				TManagedPtr<FImage> TempBase = ImagePixelFormat(CompressionQuality, Base, EImageFormat::RGBA_UByte, -1);
 				ImagePixelFormatInternal::BlockCompressionLODRangeHelper(
@@ -1346,7 +1346,7 @@ void FImageOperator::ImagePixelFormat(bool& bOutSuccess, int32 CompressionQualit
 			}
 			default:
 			{
-				UE_LOG(LogMutableCore, Log, TEXT("Image format conversion not implemented. Expensive generic one used. "));
+				UE_LOGF(LogMutableCore, Log, "Image format conversion not implemented. Expensive generic one used. ");
 
 				TManagedPtr<FImage> TempBase = ImagePixelFormat(CompressionQuality, Base, EImageFormat::RGBA_UByte, -1);
 				ImagePixelFormatInternal::BlockCompressionLODRangeHelper(
@@ -1389,7 +1389,7 @@ void FImageOperator::ImagePixelFormat(bool& bOutSuccess, int32 CompressionQualit
 
 			default:
 			{
-				UE_LOG(LogMutableCore, Log, TEXT("Image format conversion not implemented. Expensive generic one used. "));
+				UE_LOGF(LogMutableCore, Log, "Image format conversion not implemented. Expensive generic one used. ");
 
 				TManagedPtr<FImage> TempBase = ImagePixelFormat(CompressionQuality, Base, EImageFormat::RGBA_UByte, -1);
 				ImagePixelFormatInternal::BlockCompressionLODRangeHelper(
@@ -1425,7 +1425,7 @@ void FImageOperator::ImagePixelFormat(bool& bOutSuccess, int32 CompressionQualit
 			//}
 			default:
 			{
-				UE_LOG(LogMutableCore, Log, TEXT("Image format conversion not implemented. Expensive generic one used. "));
+				UE_LOGF(LogMutableCore, Log, "Image format conversion not implemented. Expensive generic one used. ");
 
 				TManagedPtr<FImage> TempBase = ImagePixelFormat(CompressionQuality, Base, EImageFormat::RGBA_UByte, -1);
 				ImagePixelFormatInternal::BlockCompressionLODRangeHelper(
@@ -1463,7 +1463,7 @@ void FImageOperator::ImagePixelFormat(bool& bOutSuccess, int32 CompressionQualit
 			//}
 			default:
 			{
-				UE_LOG(LogMutableCore, Log, TEXT("Image format conversion not implemented. Expensive generic one used. "));
+				UE_LOGF(LogMutableCore, Log, "Image format conversion not implemented. Expensive generic one used. ");
 
 				TManagedPtr<FImage> TempBase = ImagePixelFormat(CompressionQuality, Base, EImageFormat::RGBA_UByte, -1);
 				ImagePixelFormatInternal::BlockCompressionLODRangeHelper(
@@ -1505,7 +1505,7 @@ void FImageOperator::ImagePixelFormat(bool& bOutSuccess, int32 CompressionQualit
 			//}
 			default:
 			{
-				UE_LOG(LogMutableCore, Log, TEXT("Image format conversion not implemented. Expensive generic one used. "));
+				UE_LOGF(LogMutableCore, Log, "Image format conversion not implemented. Expensive generic one used. ");
 
 				TManagedPtr<FImage> TempBase = ImagePixelFormat(CompressionQuality, Base, EImageFormat::RGBA_UByte, -1);
 				ImagePixelFormatInternal::BlockCompressionLODRangeHelper(
@@ -1541,7 +1541,7 @@ void FImageOperator::ImagePixelFormat(bool& bOutSuccess, int32 CompressionQualit
 			//}
 			default:
 			{
-				UE_LOG(LogMutableCore, Log, TEXT("Image format conversion not implemented. Expensive generic one used. "));
+				UE_LOGF(LogMutableCore, Log, "Image format conversion not implemented. Expensive generic one used. ");
 
 				TManagedPtr<FImage> TempBase = ImagePixelFormat(CompressionQuality, Base, EImageFormat::RGBA_UByte, -1);
 				ImagePixelFormatInternal::BlockCompressionLODRangeHelper(
@@ -1571,7 +1571,7 @@ void FImageOperator::ImagePixelFormat(bool& bOutSuccess, int32 CompressionQualit
 			}
 			default:
 			{
-				UE_LOG(LogMutableCore, Log, TEXT("Image format conversion not implemented. Expensive generic one used. "));
+				UE_LOGF(LogMutableCore, Log, "Image format conversion not implemented. Expensive generic one used. ");
 
 				TManagedPtr<FImage> TempBase = ImagePixelFormat(CompressionQuality, Base, EImageFormat::RGBA_UByte, -1);
 				ImagePixelFormatInternal::BlockCompressionLODRangeHelper(
diff --git a/Engine/Plugins/Mutable/Source/MutableRuntime/Private/MuR/OpMeshApplyLayout.h b/Engine/Plugins/Mutable/Source/MutableRuntime/Private/MuR/OpMeshApplyLayout.h
--- a/Engine/Plugins/Mutable/Source/MutableRuntime/Private/MuR/OpMeshApplyLayout.h
+++ b/Engine/Plugins/Mutable/Source/MutableRuntime/Private/MuR/OpMeshApplyLayout.h
@@ -195,7 +195,7 @@ namespace UE::Mutable::Private
 			}
 
 			default:
-				UE_LOG(LogMutableCore, Warning, TEXT("This case is not implemented.."));
+				UE_LOGF(LogMutableCore, Warning, "This case is not implemented..");
 				check( false );
 				break;
 
diff --git a/Engine/Plugins/Mutable/Source/MutableRuntime/Private/MuR/OpMeshClipWithMesh.cpp b/Engine/Plugins/Mutable/Source/MutableRuntime/Private/MuR/OpMeshClipWithMesh.cpp
--- a/Engine/Plugins/Mutable/Source/MutableRuntime/Private/MuR/OpMeshClipWithMesh.cpp
+++ b/Engine/Plugins/Mutable/Source/MutableRuntime/Private/MuR/OpMeshClipWithMesh.cpp
@@ -547,7 +547,7 @@ namespace UE::Mutable::Private
 #if !UE_BUILD_SHIPPING
 		if (!IsMeshClosed(ClipMesh, CollapsedVertexMap))
 		{
-			UE_LOG(LogMutableCore, Warning, TEXT("Mesh operation with a mesh that is not closed as required."));
+			UE_LOGF(LogMutableCore, Warning, "Mesh operation with a mesh that is not closed as required.");
 		}
 #endif
 
diff --git a/Engine/Plugins/Mutable/Source/MutableRuntime/Private/MuR/OpMeshMerge.cpp b/Engine/Plugins/Mutable/Source/MutableRuntime/Private/MuR/OpMeshMerge.cpp
--- a/Engine/Plugins/Mutable/Source/MutableRuntime/Private/MuR/OpMeshMerge.cpp
+++ b/Engine/Plugins/Mutable/Source/MutableRuntime/Private/MuR/OpMeshMerge.cpp
@@ -1737,7 +1737,7 @@ namespace OpMeshMergeClothInternal
 			}
 			else if (pFirst->ClothSections[0].IsValid() && pSecond->ClothSections[0].IsValid())
 			{
-				UE_LOG(LogMutableCore, Error, TEXT("Merging two mesh sections with cloth is not supported."));
+				UE_LOGF(LogMutableCore, Error, "Merging two mesh sections with cloth is not supported.");
 				OpMeshMergeClothInternal::CopyClothFirst(*Result, *pFirst, *pSecond);
 			}
 		}
diff --git a/Engine/Plugins/Mutable/Source/MutableRuntime/Private/MuR/OpMeshPrepareLayout.cpp b/Engine/Plugins/Mutable/Source/MutableRuntime/Private/MuR/OpMeshPrepareLayout.cpp
--- a/Engine/Plugins/Mutable/Source/MutableRuntime/Private/MuR/OpMeshPrepareLayout.cpp
+++ b/Engine/Plugins/Mutable/Source/MutableRuntime/Private/MuR/OpMeshPrepareLayout.cpp
@@ -115,7 +115,7 @@ namespace UE::Mutable::Private
 		// Notify Overlapping layout blocks
 		if (!OverlappingBlocks.IsEmpty())
 		{
-			UE_LOG(LogMutableCore, Warning, TEXT("Mesh has %d layout block overlapping in channel %d."), OverlappingBlocks.Num() + 1, LayoutChannel);
+			UE_LOGF(LogMutableCore, Warning, "Mesh has %d layout block overlapping in channel %d.", OverlappingBlocks.Num() + 1, LayoutChannel);
 		}
 
 		// Get the information about the texture coordinates channel
@@ -127,7 +127,7 @@ namespace UE::Mutable::Private
 		if (TexCoordsBufferIndex < 0 || TexCoordsChannelIndex < 0)
 		{
 			// This is actually fine when using shared materials across LODs
-			UE_LOG(LogMutableCore, Log, TEXT("Trying to generate layout for missing UV channel %d. No layout blocks generated."), LayoutChannel);
+			UE_LOGF(LogMutableCore, Log, "Trying to generate layout for missing UV channel %d. No layout blocks generated.", LayoutChannel);
 			return false;
 		}
 
@@ -186,7 +186,7 @@ namespace UE::Mutable::Private
 			// Mutable does not support non-normalized UVs
 			if (bNonNormalizedUVs)
 			{
-				UE_LOG(LogMutableCore, Warning, TEXT("Source mesh has non-normalized UVs index %d"), LayoutChannel);
+				UE_LOGF(LogMutableCore, Warning, "Source mesh has non-normalized UVs index %d", LayoutChannel);
 			}
 		}
 
@@ -485,7 +485,7 @@ namespace UE::Mutable::Private
 		//	if (!UnassignedUVs.IsEmpty())
 		//	{
 		//		// TODO: How do we translate this to editor-time info?
-		//		UE_LOG(LogMutableCore, Warning, TEXT("Source mesh has %d vertices not assigned to any layout block in LOD %d in UVs index %d"), UnassignedUVs.Num(), LayoutChannel);
+		//		UE_LOGF(LogMutableCore, Warning, "Source mesh has %d vertices not assigned to any layout block in LOD %d in UVs index %d", UnassignedUVs.Num(), LayoutChannel);
 		//	}
 		//}
 
@@ -633,7 +633,7 @@ namespace UE::Mutable::Private
 				if (TexCoordsBufferIndex < 0 || TexCoordsChannelIndex < 0)
 				{
 					// This is actually fine when using shared materials across LODs
-					UE_LOG(LogMutableCore, Log, TEXT("Trying to generate layout for missing UV channel %d. No layout blocks generated."), LayoutChannel);
+					UE_LOGF(LogMutableCore, Log, "Trying to generate layout for missing UV channel %d. No layout blocks generated.", LayoutChannel);
 					return;
 				}
 
diff --git a/Engine/Plugins/Mutable/Source/MutableRuntime/Private/MuR/OpMeshReshape.h b/Engine/Plugins/Mutable/Source/MutableRuntime/Private/MuR/OpMeshReshape.h
--- a/Engine/Plugins/Mutable/Source/MutableRuntime/Private/MuR/OpMeshReshape.h
+++ b/Engine/Plugins/Mutable/Source/MutableRuntime/Private/MuR/OpMeshReshape.h
@@ -526,7 +526,7 @@ namespace UE::Mutable::Private
 
 		if (bTriangleOutOfScopeFound)
 		{
-			UE_LOG(LogMutableCore, Warning, TEXT("Performing a Mesh Reshape where base shape and target shape do not have the same number of triangles."));
+			UE_LOGF(LogMutableCore, Warning, "Performing a Mesh Reshape where base shape and target shape do not have the same number of triangles.");
 		}	
 #endif
 
@@ -923,7 +923,7 @@ namespace UE::Mutable::Private
 		if (!(ShapeDescriptor.Triangles.Num() && ShapeDescriptor.Normals.Num() && ShapeDescriptor.Positions.Num()))
 		{
 			bOutSuccess = false;
-			UE_LOG(LogMutableCore, Error, TEXT("The shape mesh provided does not have proper data. Skipping reshape operation."));
+			UE_LOGF(LogMutableCore, Error, "The shape mesh provided does not have proper data. Skipping reshape operation.");
 			return;
 		}
 
diff --git a/Engine/Plugins/Mutable/Source/MutableRuntime/Private/MuR/OpMeshSmoothing.h b/Engine/Plugins/Mutable/Source/MutableRuntime/Private/MuR/OpMeshSmoothing.h
--- a/Engine/Plugins/Mutable/Source/MutableRuntime/Private/MuR/OpMeshSmoothing.h
+++ b/Engine/Plugins/Mutable/Source/MutableRuntime/Private/MuR/OpMeshSmoothing.h
@@ -479,7 +479,7 @@ namespace UE::Mutable::Private
 		const bool bHasNecessaryBuffers = MeshLaplacianDataBuffer && MeshLaplacianOffsetsBuffer && MeshUniqueVertexMapBuffer;
 		if (!bHasNecessaryBuffers)
 		{
-			UE_LOG(LogMutableCore, Warning, TEXT("Smooth laplacian with source mesh missing some necessary buffers. Ignoring operation."));
+			UE_LOGF(LogMutableCore, Warning, "Smooth laplacian with source mesh missing some necessary buffers. Ignoring operation.");
 			return;
 		}
 
diff --git a/Engine/Plugins/Mutable/Source/MutableRuntime/Private/MuR/PassthroughObject.cpp b/Engine/Plugins/Mutable/Source/MutableRuntime/Private/MuR/PassthroughObject.cpp
--- a/Engine/Plugins/Mutable/Source/MutableRuntime/Private/MuR/PassthroughObject.cpp
+++ b/Engine/Plugins/Mutable/Source/MutableRuntime/Private/MuR/PassthroughObject.cpp
@@ -38,7 +38,7 @@ namespace UE::Mutable::Private
 			const TSoftObjectPtr<UObject>* Result = ResolutionMap.Find(ID);
 			if (!Result)
 			{
-				UE_LOG(LogMutableCore, Error, TEXT("Unable to find Passthrough Object with id: %u"), Pair.Key);
+				UE_LOGF(LogMutableCore, Error, "Unable to find Passthrough Object with id: %u", Pair.Key);
 				continue;
 			}
 			
@@ -73,7 +73,7 @@ namespace UE::Mutable::Private
 			UObject* Object = Result->Get();
 			if (!Object)
 			{
-				UE_LOG(LogMutableCore, Error, TEXT("Unable to load Passthrough object: %s"), *Result->ToSoftObjectPath().ToString());
+				UE_LOGF(LogMutableCore, Error, "Unable to load Passthrough object: %ls", *Result->ToSoftObjectPath().ToString());
 				ExternalObject.Resource->Set<TStrongObjectPtr<UObject>>(nullptr);
 				continue;
 			}
diff --git a/Engine/Plugins/Mutable/Source/MutableRuntime/Private/MuR/System.cpp b/Engine/Plugins/Mutable/Source/MutableRuntime/Private/MuR/System.cpp
--- a/Engine/Plugins/Mutable/Source/MutableRuntime/Private/MuR/System.cpp
+++ b/Engine/Plugins/Mutable/Source/MutableRuntime/Private/MuR/System.cpp
@@ -172,7 +172,7 @@ namespace UE::Mutable::Private
     {
     	MUTABLE_CPUPROFILER_SCOPE(NewBuildInstance);
 
-    	UE_LOG(LogMutableCore, Verbose, TEXT("Creating Mutable instance"));
+    	UE_LOGF(LogMutableCore, Verbose, "Creating Mutable instance");
     	check(InModel)
     	
     	TSharedRef<FLiveInstance> InstanceData = MakeShared<FLiveInstance>();
@@ -1751,20 +1751,20 @@ namespace UE::Mutable::Private
 		switch (Verbosity)
 		{
 		case ELogVerbosity::Error:
-			UE_LOG(LogMutableCore, Error, TEXT("%s"), *Message);
+			UE_LOGF(LogMutableCore, Error, "%ls", *Message);
 			break;
 		case ELogVerbosity::Warning:
-			UE_LOG(LogMutableCore, Warning, TEXT("%s"), *Message);
+			UE_LOGF(LogMutableCore, Warning, "%ls", *Message);
 			break;
 		case ELogVerbosity::Verbose:
-			UE_LOG(LogMutableCore, Verbose, TEXT("%s"), *Message);
+			UE_LOGF(LogMutableCore, Verbose, "%ls", *Message);
 			break;
 		case ELogVerbosity::Display:
-			UE_LOG(LogMutableCore, Verbose, TEXT("%s"), *Message);
+			UE_LOGF(LogMutableCore, Verbose, "%ls", *Message);
 			break;
 		case ELogVerbosity::Log:
 		default:
-			UE_LOG(LogMutableCore, Log, TEXT("%s"), *Message);
+			UE_LOGF(LogMutableCore, Log, "%ls", *Message);
 			break;
 		}
 
diff --git a/Engine/Plugins/Mutable/Source/MutableTools/Private/MuT/AST.cpp b/Engine/Plugins/Mutable/Source/MutableTools/Private/MuT/AST.cpp
--- a/Engine/Plugins/Mutable/Source/MutableTools/Private/MuT/AST.cpp
+++ b/Engine/Plugins/Mutable/Source/MutableTools/Private/MuT/AST.cpp
@@ -301,12 +301,12 @@ void ASTOp::LogHistogram( ASTOpList& roots )
  //       return a.Key >b.Key;
  //   });
 
- //   UE_LOG(LogMutableCore,Log, TEXT("Op histogram (%llu ops):"), count);
+ //   UE_LOGF(LogMutableCore,Log, "Op histogram (%llu ops):", count);
  //   for(int i=0; i<8; ++i)
  //   {
  //       float p = Sorted[i].Key/float(count)*100.0f;
  //       int OpType = (int)Sorted[i].Value;
- //       UE_LOG(LogMutableCore, Log, TEXT("  %5.2f%% : %s"), p, s_opNames[OpType] );
+ //       UE_LOGF(LogMutableCore, Log, "  %5.2f%% : %ls", p, s_opNames[OpType] );
  //   }
 
 	//// Debug log part of the tree
@@ -320,18 +320,18 @@ void ASTOp::LogHistogram( ASTOpList& roots )
 	//			Ptr<const ASTOpConstantResource> Constant = static_cast<const ASTOpConstantResource*>(Base.get());
 
 	//			// Log the op
-	//			UE_LOG(LogMutableCore, Log, TEXT("Multilayer at %x:"), n.get());
-	//			UE_LOG(LogMutableCore, Log, TEXT("    base is type %s:"), s_opNames[(int)Base->GetOpType()]);
+	//			UE_LOGF(LogMutableCore, Log, "Multilayer at %x:", n.get());
+	//			UE_LOGF(LogMutableCore, Log, "    base is type %ls:", s_opNames[(int)Base->GetOpType()]);
 	//			if (Constant)
 	//			{
 	//				int Format = int( ((UE::Mutable::Private::FImage*)Constant->GetValue().get())->GetFormat() );
-	//				UE_LOG(LogMutableCore, Log, TEXT("    constant image format is %d:"), Format);
+	//				UE_LOGF(LogMutableCore, Log, "    constant image format is %d:", Format);
 	//			}
 	//			else
 	//			{
 	//				FGetImageDescContext Context;
 	//				FImageDesc Desc = Base->GetImageDesc(true, &Context);
-	//				UE_LOG(LogMutableCore, Log, TEXT("    estimated image format is %d:"), int(Desc.m_format) );
+	//				UE_LOGF(LogMutableCore, Log, "    estimated image format is %d:", int(Desc.m_format) );
 	//			}
 	//		}
 	//		return true;
diff --git a/Engine/Plugins/Mutable/Source/MutableTools/Private/MuT/CodeOptimiser.cpp b/Engine/Plugins/Mutable/Source/MutableTools/Private/MuT/CodeOptimiser.cpp
--- a/Engine/Plugins/Mutable/Source/MutableTools/Private/MuT/CodeOptimiser.cpp
+++ b/Engine/Plugins/Mutable/Source/MutableTools/Private/MuT/CodeOptimiser.cpp
@@ -693,36 +693,36 @@ namespace UE::Mutable::Private
 			{
 				--OptimizeIterationsLeft;
 				++NumIterations;
-				UE_LOG(LogMutableCore, Verbose, TEXT("Main optimise iteration %d, left %d"), NumIterations, OptimizeIterationsLeft);
+				UE_LOGF(LogMutableCore, Verbose, "Main optimise iteration %d, left %d", NumIterations, OptimizeIterationsLeft);
 
 				bModifiedInInnerLoop = false;
 
 				// All kind of optimisations that depend on the meaning of each operation
 				// \TODO: We are doing it for all states.
-				UE_LOG(LogMutableCore, Verbose, TEXT(" - semantic optimiser"));
+				UE_LOGF(LogMutableCore, Verbose, " - semantic optimiser");
 				bModifiedInInnerLoop |= SemanticOptimiserAST(roots, Options->GetPrivate()->OptimisationOptions, Pass);
 				ASTOp::LogHistogram(roots);
 
-				UE_LOG(LogMutableCore, Verbose, TEXT(" - sink optimiser"));
+				UE_LOGF(LogMutableCore, Verbose, " - sink optimiser");
 				bModifiedInInnerLoop |= SinkOptimiserAST(roots, Options->GetPrivate()->OptimisationOptions);
 				ASTOp::LogHistogram(roots);
 
 				// Image size operations are treated separately
-				UE_LOG(LogMutableCore, Verbose, TEXT(" - size optimiser"));
+				UE_LOGF(LogMutableCore, Verbose, " - size optimiser");
 				bModifiedInInnerLoop |= SizeOptimiserAST(roots);
 			}
 
 			bModified = bModifiedInInnerLoop;
 
-			UE_LOG(LogMutableCore, Verbose, TEXT(" - duplicated code remover"));
+			UE_LOGF(LogMutableCore, Verbose, " - duplicated code remover");
 			bModified |= DuplicatedCodeRemoverAST(roots);
-			//UE_LOG(LogMutableCore, Verbose, TEXT("(int) %s : %ld"), TEXT("ast size"), int64(ASTOp::CountNodes(roots)));
+			//UE_LOGF(LogMutableCore, Verbose, TEXT("(int) %ls : %ld"), "ast size", int64(ASTOp::CountNodes(roots)));
 
 			ASTOp::LogHistogram(roots);
 
-			UE_LOG(LogMutableCore, Verbose, TEXT(" - duplicated data remover"));
+			UE_LOGF(LogMutableCore, Verbose, " - duplicated data remover");
 			bModified |= DuplicatedDataRemoverAST(roots);
-			//UE_LOG(LogMutableCore, Verbose, TEXT("(int) %s : %ld"), TEXT("ast size"), int64(ASTOp::CountNodes(roots)));
+			//UE_LOGF(LogMutableCore, Verbose, TEXT("(int) %ls : %ld"), "ast size", int64(ASTOp::CountNodes(roots)));
 
 			ASTOp::LogHistogram(roots);
 
@@ -730,8 +730,8 @@ namespace UE::Mutable::Private
 			bool bModifiedInConstants = false;
 			for (Ptr<ASTOp>& Root : roots)
 			{
-				//UE_LOG(LogMutableCore, Verbose, TEXT("(int) %s : %ld"), TEXT("ast size"), int64(ASTOp::CountNodes(roots)));
-				UE_LOG(LogMutableCore, Verbose, TEXT(" - constant generator"));
+				//UE_LOGF(LogMutableCore, Verbose, TEXT("(int) %ls : %ld"), "ast size", int64(ASTOp::CountNodes(roots)));
+				UE_LOGF(LogMutableCore, Verbose, " - constant generator");
 
 				// Constant subtree generation
 				bModifiedInConstants |= ConstantGenerator(Options->GetPrivate(), Root, Pass);
@@ -743,13 +743,13 @@ namespace UE::Mutable::Private
 			{
 				bModified = true;
 
-				UE_LOG(LogMutableCore, Verbose, TEXT(" - duplicated data remover"));
+				UE_LOGF(LogMutableCore, Verbose, " - duplicated data remover");
 				DuplicatedDataRemoverAST(roots);
 			}
 
 			//if (!bModified)
 			{
-				UE_LOG(LogMutableCore, Verbose, TEXT(" - logic optimiser"));
+				UE_LOGF(LogMutableCore, Verbose, " - logic optimiser");
 				bModified |= LocalLogicOptimiserAST(roots);
 			}
 
@@ -955,7 +955,7 @@ namespace UE::Mutable::Private
 			roots.Add(s.root);
 		}
 
-		//UE_LOG(LogMutableCore, Verbose, TEXT("(int) %s : %ld"), TEXT("ast size"), int64(ASTOp::CountNodes(roots)));
+		//UE_LOGF(LogMutableCore, Verbose, TEXT("(int) %ls : %ld"), "ast size", int64(ASTOp::CountNodes(roots)));
 
 		if ( Options->GetPrivate()->OptimisationOptions.bEnabled )
 		{
@@ -968,18 +968,18 @@ namespace UE::Mutable::Private
 			// duplicated data (meshes) that may have been specified in the source
 			// data, before we make it diverge because of different uses, like layout
 			// creation
-			UE_LOG(LogMutableCore, Verbose, TEXT(" - duplicated data remover"));
+			UE_LOGF(LogMutableCore, Verbose, " - duplicated data remover");
 			DuplicatedDataRemoverAST( roots );
 
 			ASTOp::LogHistogram(roots);
 
-			UE_LOG(LogMutableCore, Verbose, TEXT(" - duplicated code remover"));
+			UE_LOGF(LogMutableCore, Verbose, " - duplicated code remover");
 			DuplicatedCodeRemoverAST( roots );
 
 			// Special optimization stages
 			if ( Options->GetPrivate()->OptimisationOptions.bUniformizeSkeleton )
 			{
-				UE_LOG(LogMutableCore, Verbose, TEXT(" - skeleton cleaner"));
+				UE_LOGF(LogMutableCore, Verbose, " - skeleton cleaner");
 				ASTOp::LogHistogram(roots);
 
 				SkeletonCleanerAST( roots, Options->GetPrivate()->OptimisationOptions );
@@ -997,11 +997,11 @@ namespace UE::Mutable::Private
 
 				--OptimizeIterationsLeft;
 				++NumIterations;
-				UE_LOG(LogMutableCore, Verbose, TEXT("First optimise iteration %d, left %d"), NumIterations, OptimizeIterationsLeft);
+				UE_LOGF(LogMutableCore, Verbose, "First optimise iteration %d, left %d", NumIterations, OptimizeIterationsLeft);
 
 				bModified = false;
 
-				UE_LOG(LogMutableCore, Verbose, TEXT(" - size optimiser"));
+				UE_LOGF(LogMutableCore, Verbose, " - size optimiser");
 				bModified |= SizeOptimiserAST( roots );
 			}
 
@@ -1081,10 +1081,10 @@ namespace UE::Mutable::Private
 			// duplicated data (meshes) that may have been specified in the source
 			// data, before we make it diverge because of different uses, like layout
 			// creation
-			UE_LOG(LogMutableCore, Verbose, TEXT(" - duplicated data remover"));
+			UE_LOGF(LogMutableCore, Verbose, " - duplicated data remover");
 			DuplicatedDataRemoverAST( roots );
 
-			UE_LOG(LogMutableCore, Verbose, TEXT(" - duplicated code remover"));
+			UE_LOGF(LogMutableCore, Verbose, " - duplicated code remover");
 			DuplicatedCodeRemoverAST( roots );
 
 			// Constant resolution stage: resolve referenced assets.
@@ -1097,14 +1097,14 @@ namespace UE::Mutable::Private
 			{
 				constexpr int32 Pass = 1;
 
-				UE_LOG(LogMutableCore, Verbose, TEXT(" - constant generator"));
+				UE_LOGF(LogMutableCore, Verbose, " - constant generator");
 				ConstantGenerator( Options->GetPrivate(), roots[StateIndex], Pass);
 			}
 
-			UE_LOG(LogMutableCore, Verbose, TEXT(" - duplicated data remover"));
+			UE_LOGF(LogMutableCore, Verbose, " - duplicated data remover");
 			DuplicatedDataRemoverAST( roots );
 
-			UE_LOG(LogMutableCore, Verbose, TEXT(" - duplicated code remover"));
+			UE_LOGF(LogMutableCore, Verbose, " - duplicated code remover");
 			DuplicatedCodeRemoverAST( roots );
 
 			// Reset the state root operations in case they have changed due to optimization
diff --git a/Engine/Plugins/Mutable/Source/MutableTools/Private/MuT/CodeOptimiser_Parameters.cpp b/Engine/Plugins/Mutable/Source/MutableTools/Private/MuT/CodeOptimiser_Parameters.cpp
--- a/Engine/Plugins/Mutable/Source/MutableTools/Private/MuT/CodeOptimiser_Parameters.cpp
+++ b/Engine/Plugins/Mutable/Source/MutableTools/Private/MuT/CodeOptimiser_Parameters.cpp
@@ -2063,9 +2063,9 @@ namespace UE::Mutable::Private
 
                     ++NumIterations;
                     --OptimizeIterationsLeft;
-                    UE_LOG(LogMutableCore, Verbose, TEXT("State optimise iteration %d, left %d"), NumIterations, OptimizeIterationsLeft);
+                    UE_LOGF(LogMutableCore, Verbose, "State optimise iteration %d, left %d", NumIterations, OptimizeIterationsLeft);
 
-                    UE_LOG(LogMutableCore, Verbose, TEXT(" - before parameter optimiser"));
+                    UE_LOGF(LogMutableCore, Verbose, " - before parameter optimiser");
 
                     ParameterOptimiserAST param( States[s], Options->GetPrivate()->OptimisationOptions );
 					bModified = param.Apply();
@@ -2073,17 +2073,17 @@ namespace UE::Mutable::Private
                     TArray<Ptr<ASTOp>> Roots;
                     Roots.Add(States[s].root);
 
-                    UE_LOG(LogMutableCore, Verbose, TEXT(" - after parameter optimiser"));
+                    UE_LOGF(LogMutableCore, Verbose, " - after parameter optimiser");
 
                     // All kind of optimisations that depend on the meaning of each operation
-                    UE_LOG(LogMutableCore, Verbose, TEXT(" - semantic optimiser"));
+                    UE_LOGF(LogMutableCore, Verbose, " - semantic optimiser");
 					bModified |= SemanticOptimiserAST(Roots, Options->GetPrivate()->OptimisationOptions, 1 );
 
-                    UE_LOG(LogMutableCore, Verbose, TEXT(" - sink optimiser"));
+                    UE_LOGF(LogMutableCore, Verbose, " - sink optimiser");
 					bModified |= SinkOptimiserAST(Roots, Options->GetPrivate()->OptimisationOptions );
 
                     // Image size operations are treated separately
-                    UE_LOG(LogMutableCore, Verbose, TEXT(" - size optimiser"));
+                    UE_LOGF(LogMutableCore, Verbose, " - size optimiser");
 					bModified |= SizeOptimiserAST(Roots);
 
 					// Some sink optimizations can only be applied after some constant reductions
@@ -2097,10 +2097,10 @@ namespace UE::Mutable::Private
                 TArray<Ptr<ASTOp>> Roots;
                 Roots.Add(States[s].root);
 
-                UE_LOG(LogMutableCore, Verbose, TEXT(" - duplicated data remover"));
+                UE_LOGF(LogMutableCore, Verbose, " - duplicated data remover");
 				bModified |= DuplicatedDataRemoverAST(Roots);
 
-                UE_LOG(LogMutableCore, Verbose, TEXT(" - duplicated code remover"));
+                UE_LOGF(LogMutableCore, Verbose, " - duplicated code remover");
 				bModified |= DuplicatedCodeRemoverAST(Roots);
 
                 States[s].root = Roots[0];
@@ -2129,13 +2129,13 @@ namespace UE::Mutable::Private
             for (FStateCompilationData& s: States )
             {
                 {
-                    UE_LOG(LogMutableCore, Verbose, TEXT(" - state cache"));
+                    UE_LOGF(LogMutableCore, Verbose, " - state cache");
                     MUTABLE_CPUPROFILER_SCOPE(StateCache);
                     StateCacheDetectorAST c( &s );
                 }
 
                 {
-                    UE_LOG(LogMutableCore, Verbose, TEXT(" - state cache format"));
+                    UE_LOGF(LogMutableCore, Verbose, " - state cache format");
                     MUTABLE_CPUPROFILER_SCOPE(StateCacheFormat);
                     StateCacheFormatOptimiserAST f( s, opFormats );
                 }
@@ -2152,28 +2152,28 @@ namespace UE::Mutable::Private
             {
                 ++NumIterations;
                 --OptimizeIterationsLeft;
-                UE_LOG(LogMutableCore, Verbose, TEXT("State reoptimise iteration %d, left %d"), NumIterations, OptimizeIterationsLeft);
+                UE_LOGF(LogMutableCore, Verbose, "State reoptimise iteration %d, left %d", NumIterations, OptimizeIterationsLeft);
 
 				bModified = false;
 
-                UE_LOG(LogMutableCore, Verbose, TEXT(" - semantic optimiser"));
+                UE_LOGF(LogMutableCore, Verbose, " - semantic optimiser");
 				bModified |= SemanticOptimiserAST( Roots, Options->GetPrivate()->OptimisationOptions, Pass );
 
                 // Image size operations are treated separately
-                UE_LOG(LogMutableCore, Verbose, TEXT(" - size optimiser"));
+                UE_LOGF(LogMutableCore, Verbose, " - size optimiser");
 				bModified |= SizeOptimiserAST( Roots );
 			}
 
             for(Ptr<ASTOp>& Root : Roots)
             {
-                UE_LOG(LogMutableCore, Verbose, TEXT(" - constant optimiser"));
+                UE_LOGF(LogMutableCore, Verbose, " - constant optimiser");
 				bModified = ConstantGenerator( Options->GetPrivate(), Root, Pass );
 			}
 
-            UE_LOG(LogMutableCore, Verbose, TEXT(" - duplicated data remover"));
+            UE_LOGF(LogMutableCore, Verbose, " - duplicated data remover");
             DuplicatedDataRemoverAST( Roots );
 
-            UE_LOG(LogMutableCore, Verbose, TEXT(" - duplicated code remover"));
+            UE_LOGF(LogMutableCore, Verbose, " - duplicated code remover");
             DuplicatedCodeRemoverAST( Roots );
 		}
 
diff --git a/Engine/Plugins/Mutable/Source/MutableTools/Private/MuT/Compiler.cpp b/Engine/Plugins/Mutable/Source/MutableTools/Private/MuT/Compiler.cpp
--- a/Engine/Plugins/Mutable/Source/MutableTools/Private/MuT/Compiler.cpp
+++ b/Engine/Plugins/Mutable/Source/MutableTools/Private/MuT/Compiler.cpp
@@ -685,15 +685,15 @@ namespace UE::Mutable::Private
 			Program.Serialise(ProgramMemoryArch);
 		}
 
-		UE_LOG(LogMutableCore, Log, TEXT("Compilation Time: %lfs"), FPlatformTime::Seconds() - StartTimeSeconds);
+		UE_LOGF(LogMutableCore, Log, "Compilation Time: %lfs", FPlatformTime::Seconds() - StartTimeSeconds);
 
 		const double ProgramSize  = double(ProgramSizeStream.GetBufferSize());
 		const double ByteCodeSize = double(Program.ByteCode.Num());
 		constexpr double ByteToMiB = 1.0/(1024.0*1024.0); 
 
-		UE_LOG(LogMutableCore, Log, TEXT("Program Size : %lf MiB"), ProgramSize * ByteToMiB);
-		UE_LOG(LogMutableCore, Log, TEXT("Program ByteCode Size: %lf MiB"), ByteCodeSize * ByteToMiB);
-		UE_LOG(LogMutableCore, Log, TEXT("Program Num Ops : %d"), Program.NumOps);
+		UE_LOGF(LogMutableCore, Log, "Program Size : %lf MiB", ProgramSize * ByteToMiB);
+		UE_LOGF(LogMutableCore, Log, "Program ByteCode Size: %lf MiB", ByteCodeSize * ByteToMiB);
+		UE_LOGF(LogMutableCore, Log, "Program Num Ops : %d", Program.NumOps);
 
         return Result;
     }
@@ -982,7 +982,7 @@ namespace UE::Mutable::Private
 			}
 		}
 
-		UE_LOG(LogMutableCore, Log, TEXT("Generated roms: %d (%d KB) are embedded, %d (%d KB) are streamed of which %d (%d KB) are high-res."), 
+		UE_LOGF(LogMutableCore, Log, "Generated roms: %d (%d KB) are embedded, %d (%d KB) are streamed of which %d (%d KB) are high-res.", 
 			NumEmbedded, uint32(NumEmbeddedBytes/1024), NumRoms, uint32(NumRomsBytes/1024), NumHighRes, uint32(NumHighResBytes/1024));
 	}
 
diff --git a/Engine/Plugins/Mutable/Source/MutableTools/Private/MuT/ErrorLog.cpp b/Engine/Plugins/Mutable/Source/MutableTools/Private/MuT/ErrorLog.cpp
--- a/Engine/Plugins/Mutable/Source/MutableTools/Private/MuT/ErrorLog.cpp
+++ b/Engine/Plugins/Mutable/Source/MutableTools/Private/MuT/ErrorLog.cpp
@@ -164,16 +164,16 @@ namespace UE::Mutable::Private
 	
 	void FErrorLog::Log() const
 	{
-		UE_LOG(LogMutableCore, Log, TEXT(" Error Log :\n"));
+		UE_LOGF(LogMutableCore, Log, " Error Log :\n");
 
 		for ( const FMessage& msg : Messages )
 		{
 			switch ( msg.Type )
 			{
-			case ELMT_ERROR: 	UE_LOG(LogMutableCore, Log, TEXT("  ERR  %s\n"), *msg.Text); break;
-			case ELMT_WARNING: 	UE_LOG(LogMutableCore, Log, TEXT("  WRN  %s\n"), *msg.Text); break;
-			case ELMT_INFO: 	UE_LOG(LogMutableCore, Log, TEXT("  INF  %s\n"), *msg.Text); break;
-			default: 			UE_LOG(LogMutableCore, Log, TEXT("  NON  %s\n"), *msg.Text); break;
+			case ELMT_ERROR: 	UE_LOGF(LogMutableCore, Log, "  ERR  %ls\n", *msg.Text); break;
+			case ELMT_WARNING: 	UE_LOGF(LogMutableCore, Log, "  WRN  %ls\n", *msg.Text); break;
+			case ELMT_INFO: 	UE_LOGF(LogMutableCore, Log, "  INF  %ls\n", *msg.Text); break;
+			default: 			UE_LOGF(LogMutableCore, Log, "  NON  %ls\n", *msg.Text); break;
 			}
 		}
 	}
diff --git a/Engine/Plugins/Mutable/Source/MutableTools/Private/MuT/UnrealPixelFormatOverride.cpp b/Engine/Plugins/Mutable/Source/MutableTools/Private/MuT/UnrealPixelFormatOverride.cpp
--- a/Engine/Plugins/Mutable/Source/MutableTools/Private/MuT/UnrealPixelFormatOverride.cpp
+++ b/Engine/Plugins/Mutable/Source/MutableTools/Private/MuT/UnrealPixelFormatOverride.cpp
@@ -260,7 +260,7 @@ bool ImageCoreToMutable(const FCompressedImage2D& Compressed, UE::Mutable::Priva
 
 	if (Compressed.RawData.Num() != MutableView.Num())
 	{
-		UE_LOG(LogMutableCore, Error, TEXT("Buffer size mismatch when trying to convert image LOD %d, mutable size is %d and ue size is %d. Mutable is %d x %d format %d and UE is %d x %d format %d."), 
+		UE_LOGF(LogMutableCore, Error, "Buffer size mismatch when trying to convert image LOD %d, mutable size is %d and ue size is %d. Mutable is %d x %d format %d and UE is %d x %d format %d.", 
 			LOD, MutableView.Num(), Compressed.RawData.Num(), 
 			Mutable->GetSizeX(), Mutable->GetSizeY(), Mutable->GetFormat(),
 			Compressed.SizeX, Compressed.SizeY, Compressed.PixelFormat
