[Mutable] Change all instances of UE_LOG to UE_LOGF in CustomizableObjectEditor module.

#rnx
#Synced-CL 51716106

diff --git a/Engine/Plugins/Mutable/Source/CustomizableObjectEditor/Private/MuCOE/CustomizableObjectBenchmarkingUtils.cpp b/Engine/Plugins/Mutable/Source/CustomizableObjectEditor/Private/MuCOE/CustomizableObjectBenchmarkingUtils.cpp
--- a/Engine/Plugins/Mutable/Source/CustomizableObjectEditor/Private/MuCOE/CustomizableObjectBenchmarkingUtils.cpp
+++ b/Engine/Plugins/Mutable/Source/CustomizableObjectEditor/Private/MuCOE/CustomizableObjectBenchmarkingUtils.cpp
@@ -22,12 +22,12 @@ bool CustomizableObjectBenchmarkingUtils::GenerateDeterministicSetOfInstances(UC
 	const uint32 StateCount = TargetCustomizableObject.GetStateCount();
 	check(StateCount >= 1);
 
-	UE_LOG(LogMutable, Display, TEXT("Requested Instances Count : %i"), InstancesPerState);
-	UE_LOG(LogMutable, Display, TEXT("State Count = %i"), StateCount);
+	UE_LOGF(LogMutable, Display, "Requested Instances Count : %i", InstancesPerState);
+	UE_LOGF(LogMutable, Display, "State Count = %i", StateCount);
 	
 	// Compute actual total amount of instances to generate
 	const uint32 TotalInstancesToTestCount = InstancesPerState * StateCount;
-	UE_LOG(LogMutable,Display,TEXT("Generating %i instances (states * requested instances)..."), TotalInstancesToTestCount);
+	UE_LOGF(LogMutable,Display, "Generating %i instances (states * requested instances)...", TotalInstancesToTestCount);
 	
 	// Create randomization stream for the parameters of the instance
 	FRandomStream RandomizationStream = FRandomStream(0);
@@ -50,7 +50,7 @@ bool CustomizableObjectBenchmarkingUtils::GenerateDeterministicSetOfInstances(UC
 		}
 		else
 		{
-			UE_LOG(LogMutable,Error,TEXT("Failed to generate COI for the %s CO."),*TargetCustomizableObject.GetName());
+			UE_LOGF(LogMutable,Error, "Failed to generate COI for the %ls CO.",*TargetCustomizableObject.GetName());
 			return false;
 		}
 	}
diff --git a/Engine/Plugins/Mutable/Source/CustomizableObjectEditor/Private/MuCOE/CustomizableObjectCompileRunnable.cpp b/Engine/Plugins/Mutable/Source/CustomizableObjectEditor/Private/MuCOE/CustomizableObjectCompileRunnable.cpp
--- a/Engine/Plugins/Mutable/Source/CustomizableObjectEditor/Private/MuCOE/CustomizableObjectCompileRunnable.cpp
+++ b/Engine/Plugins/Mutable/Source/CustomizableObjectEditor/Private/MuCOE/CustomizableObjectCompileRunnable.cpp
@@ -75,7 +75,7 @@ UE::Mutable::Private::TManagedPtr<UE::Mutable::Private::FImage> FCustomizableObj
 	{
 		// This could happen in the editor, because some source textures may have changed while there was a background compilation.
 		// We just show a warning and move on. This cannot happen during cooks, so it is fine.
-		UE_LOG(LogMutable, Warning, TEXT("Failed to load some source texture data for texture ID [%d]. Some textures may be corrupted."), ID);
+		UE_LOGF(LogMutable, Warning, "Failed to load some source texture data for texture ID [%d]. Some textures may be corrupted.", ID);
 	}
 
 	return Image;
@@ -86,7 +86,7 @@ uint32 FCustomizableObjectCompileRunnable::Run()
 {
 	TRACE_BEGIN_REGION(UE_MUTABLE_CORE_REGION);
 
-	UE_LOG(LogMutable, Verbose, TEXT("PROFILE: [ %16.8f ] FCustomizableObjectCompileRunnable::Run start."), FPlatformTime::Seconds());
+	UE_LOGF(LogMutable, Verbose, "PROFILE: [ %16.8f ] FCustomizableObjectCompileRunnable::Run start.", FPlatformTime::Seconds());
 
 	uint32 Result = 1;
 	ErrorMsg = FString();
@@ -96,7 +96,7 @@ uint32 FCustomizableObjectCompileRunnable::Run()
 
 	if (Options.OptimizationLevel > UE_MUTABLE_MAX_OPTIMIZATION)
 	{
-		UE_LOG(LogMutable, Log, TEXT("Mutable compile optimization level out of range. Clamping to maximum."));
+		UE_LOGF(LogMutable, Log, "Mutable compile optimization level out of range. Clamping to maximum.");
 		Options.OptimizationLevel = UE_MUTABLE_MAX_OPTIMIZATION;
 	}
 
@@ -334,7 +334,7 @@ uint32 FCustomizableObjectCompileRunnable::Run()
 
 	UE::Mutable::Private::Ptr<UE::Mutable::Private::Compiler> Compiler = new UE::Mutable::Private::Compiler(CompilerOptions, WaitCallback);
 
-	UE_LOG(LogMutable, Verbose, TEXT("PROFILE: [ %16.8f ] FCustomizableObjectCompileRunnable Compile start."), FPlatformTime::Seconds());
+	UE_LOGF(LogMutable, Verbose, "PROFILE: [ %16.8f ] FCustomizableObjectCompileRunnable Compile start.", FPlatformTime::Seconds());
 	Model = Compiler->Compile(MutableRoot, ExternalOperations);
 
 	// Dump all the log messages from the compiler
@@ -382,7 +382,7 @@ uint32 FCustomizableObjectCompileRunnable::Run()
 
 	bThreadCompleted = true;
 
-	UE_LOG(LogMutable, Verbose, TEXT("PROFILE: [ %16.8f ] FCustomizableObjectCompileRunnable::Run end."), FPlatformTime::Seconds());
+	UE_LOGF(LogMutable, Verbose, "PROFILE: [ %16.8f ] FCustomizableObjectCompileRunnable::Run end.", FPlatformTime::Seconds());
 
 	TRACE_END_REGION(UE_MUTABLE_CORE_REGION);
 
@@ -684,7 +684,7 @@ void FCustomizableObjectSaveDDRunnable::StoreCachedPlatformDataToDisk()
 		if (FileManager.FileExists(*FilePath)
 			&& !FileManager.Delete(*FilePath, true, false, true))
 		{
-			UE_LOG(LogMutable, Error, TEXT("Failed to delete file for data type [%d]."), DataType);
+			UE_LOGF(LogMutable, Error, "Failed to delete file for data type [%d].", DataType);
 			bSuccess = false;
 			break;
 		}
@@ -737,7 +737,7 @@ void FCustomizableObjectSaveDDRunnable::StoreCachedPlatformDataToDisk()
 
 				if (!MemoryWriter->Close())
 				{
-					UE_LOG(LogMutable, Error, TEXT("Failed to write file to disk. File [%s]."), *MemoryWriter->GetArchiveName());
+					UE_LOGF(LogMutable, Error, "Failed to write file to disk. File [%ls].", *MemoryWriter->GetArchiveName());
 					bSuccess = false;
 					break;
 				}
diff --git a/Engine/Plugins/Mutable/Source/CustomizableObjectEditor/Private/MuCOE/CustomizableObjectCompiler.cpp b/Engine/Plugins/Mutable/Source/CustomizableObjectEditor/Private/MuCOE/CustomizableObjectCompiler.cpp
--- a/Engine/Plugins/Mutable/Source/CustomizableObjectEditor/Private/MuCOE/CustomizableObjectCompiler.cpp
+++ b/Engine/Plugins/Mutable/Source/CustomizableObjectEditor/Private/MuCOE/CustomizableObjectCompiler.cpp
@@ -151,7 +151,7 @@ int32 FCustomizableObjectCompiler::GetNumCompileRequests() const
 void FCustomizableObjectCompiler::PreloadReferencerAssets()
 {
 	TRACE_BEGIN_REGION(UE_MUTABLE_PRELOAD_REGION);
-	UE_LOG(LogMutable, Verbose, TEXT("PROFILE: [ %16.8f ] Preload asynchronously assets start."), FPlatformTime::Seconds());
+	UE_LOGF(LogMutable, Verbose, "PROFILE: [ %16.8f ] Preload asynchronously assets start.", FPlatformTime::Seconds());
 
 	TArray<FAssetData> ReferencingAssets;
 	GetReferencingPackages(*CurrentObject, ReferencingAssets);
@@ -212,7 +212,7 @@ void FCustomizableObjectCompiler::PreloadingReferencerAssetsCallback(bool bAsync
 		AsynchronousStreamableHandlePtr = nullptr;
 	}
 
-	UE_LOG(LogMutable, Verbose, TEXT("PROFILE: [ %16.8f ] Preload asynchronously assets end."), FPlatformTime::Seconds());
+	UE_LOGF(LogMutable, Verbose, "PROFILE: [ %16.8f ] Preload asynchronously assets end.", FPlatformTime::Seconds());
 	TRACE_END_REGION(UE_MUTABLE_PRELOAD_REGION);
 
 	CompileInternal(bAsync);
@@ -235,7 +235,7 @@ void FCustomizableObjectCompiler::Compile(const TSharedRef<FCompilationRequest>&
 	UCustomizableObject* Object = InCompileRequest->GetCustomizableObject();
 	if (!Object)
 	{
-		UE_LOG(LogMutable, Warning, TEXT("Failed to compile Customizable Object. Object is missing."));
+		UE_LOGF(LogMutable, Warning, "Failed to compile Customizable Object. Object is missing.");
 		
 		FCompileCallbackParams Params;
 		Params.bRequestFailed = true;
@@ -351,7 +351,7 @@ void FCustomizableObjectCompiler::Compile(const TSharedRef<FCompilationRequest>&
 	if (!System->LockObject(CurrentObject))
 	{
 		FString Message = FString::Printf(TEXT("Customizable Object %s is already being compiled or updated. Please wait a few seconds and try again."), *CurrentObject->GetName());
-		UE_LOG(LogMutable, Warning, TEXT("%s"), *Message);
+		UE_LOGF(LogMutable, Warning, "%ls", *Message);
 		
 		FNotificationInfo Info(LOCTEXT("CustomizableObjectBeingCompilerOrUpdated", "Customizable Object compile and/or update still in process. Please wait a few seconds and try again."));
 		Info.bFireAndForget = true;
@@ -393,18 +393,18 @@ void FCustomizableObjectCompiler::Compile(const TSharedRef<FCompilationRequest>&
 
 		if (IsRunningCommandlet())
 		{
-			UE_LOG(LogMutable, Display, TEXT("%s"), *MutableCompilationStartMessage);
+			UE_LOGF(LogMutable, Display, "%ls", *MutableCompilationStartMessage);
 		}
 		else
 		{
 			const int32 MaxOptimizationLevelValue = ConvertOptimizationLevel(ECustomizableObjectOptimizationLevel::Maximum);
 			if (CurrentOptions.OptimizationLevel == MaxOptimizationLevelValue)
 			{
-				UE_LOG(LogMutable, Display, TEXT("%s The Compilation will take more time to run due to the chosen optimization level."), *MutableCompilationStartMessage);
+				UE_LOGF(LogMutable, Display, "%ls The Compilation will take more time to run due to the chosen optimization level.", *MutableCompilationStartMessage);
 			}
 			else
 			{
-				UE_LOG(LogMutable, Display, TEXT("%s"), *MutableCompilationStartMessage);
+				UE_LOGF(LogMutable, Display, "%ls", *MutableCompilationStartMessage);
 			}
 		}
 		
@@ -412,7 +412,7 @@ void FCustomizableObjectCompiler::Compile(const TSharedRef<FCompilationRequest>&
 
 	if (CurrentOptions.bForceLargeLODBias)
 	{
-		UE_LOG(LogMutable, Display, TEXT("Compiling Customizable Object with %d LODBias."), CurrentOptions.DebugBias);
+		UE_LOGF(LogMutable, Display, "Compiling Customizable Object with %d LODBias.", CurrentOptions.DebugBias);
 	}
 
 	// Create and update compilation progress notification
@@ -464,7 +464,7 @@ void FCustomizableObjectCompiler::EnqueueCompileRequest(const TSharedRef<FCompil
 			// Wait for the search to be completed if this new compilation is SYNC
 			if (AssetRegistry.IsGathering() && !CompileRequest->bAsync)
 			{
-				UE_LOG(LogMutable, Log, TEXT("Waiting for the AR scan to be completed before enqueueing the compilation of the %s CO."), *CustomizableObject->GetName());
+				UE_LOGF(LogMutable, Log, "Waiting for the AR scan to be completed before enqueueing the compilation of the %ls CO.", *CustomizableObject->GetName());
 				AssetRegistry.WaitForCompletion();
 			}
 		}
@@ -474,7 +474,7 @@ void FCustomizableObjectCompiler::EnqueueCompileRequest(const TSharedRef<FCompil
 			// In that case, force the search to be synchronous to ensure the compilation can proceed.
 			const bool bShouldSearchAssetsSync = (IsRunningCommandlet() && !IsRunningCookCommandlet()) ? true : !CompileRequest->bAsync;
 			
-			UE_LOG(LogMutable, Display, TEXT("Performing full %s Asset Registry search as required by the CO compilation process."), bShouldSearchAssetsSync ? *FString(TEXT("SYNC")) : *FString(TEXT("ASYNC")));
+			UE_LOGF(LogMutable, Display, "Performing full %s Asset Registry search as required by the CO compilation process.", "ASYNC");
 			
 			// Scan for all the assets async/sync based on the compilation options as the AR has not yet been scanned
 			AssetRegistry.SearchAllAssets(bShouldSearchAssetsSync);
@@ -591,7 +591,7 @@ void ProcessChildObjectsRecursively(const UCustomizableObject* ParentObject, FMu
 		{
 			if (!GenerationContext.RootVersionBridge)
 			{
-				UE_LOG(LogMutable, Warning, TEXT("The child Customizable Object [%s] defines its VersionStruct Property but its root CustomizableObject doesn't define the VersionBridge property. There's no way to verify the VersionStruct has to be included in this compilation, so the child CustomizableObject will be omitted."), 
+				UE_LOGF(LogMutable, Warning, "The child Customizable Object [%ls] defines its VersionStruct Property but its root CustomizableObject doesn't define the VersionBridge property. There's no way to verify the VersionStruct has to be included in this compilation, so the child CustomizableObject will be omitted.", 
 					*ChildObject->GetName());
 				continue;
 			}
@@ -648,7 +648,7 @@ UE::Mutable::Private::Ptr<UE::Mutable::Private::NodeObject> GenerateMutableRoot(
 
 		if (IsRunningCookCommandlet() || IsRunningCookOnTheFly())
 		{
-			UE_LOG(LogMutable, Warning, TEXT("Compilation failed! Missing EDITORONLY data for Customizable Object [%s]. The object might have been loaded outside the Cooking context."), *Object->GetName());
+			UE_LOGF(LogMutable, Warning, "Compilation failed! Missing EDITORONLY data for Customizable Object [%ls]. The object might have been loaded outside the Cooking context.", *Object->GetName());
 		}
 
 		return nullptr;
@@ -691,12 +691,12 @@ UE::Mutable::Private::Ptr<UE::Mutable::Private::NodeObject> GenerateMutableRoot(
 		// We cannot load while saving. This should only happen in cooking and all assets should have been preloaded.
 		if (!UE::IsSavingPackage())
 		{
-			UE_LOG(LogMutable, Verbose, TEXT("PROFILE: [ %16.8f ] Begin search for children."), FPlatformTime::Seconds());
+			UE_LOGF(LogMutable, Verbose, "PROFILE: [ %16.8f ] Begin search for children.", FPlatformTime::Seconds());
 
 			// The object doesn't reference a root object but is a root object, look for all the objects that reference it and get their root nodes
 			FAssetRegistryModule& AssetRegistryModule = FModuleManager::LoadModuleChecked<FAssetRegistryModule>("AssetRegistry");
 			ProcessChildObjectsRecursively(RootObject, GenerationContext);
-			UE_LOG(LogMutable, Verbose, TEXT("PROFILE: [ %16.8f ] End search for children."), FPlatformTime::Seconds());
+			UE_LOGF(LogMutable, Verbose, "PROFILE: [ %16.8f ] End search for children.", FPlatformTime::Seconds());
 		}
 	}
 	else
@@ -776,9 +776,9 @@ UE::Mutable::Private::Ptr<UE::Mutable::Private::NodeObject> GenerateMutableRoot(
 	GenerationContext.bPartialCompilation = LocalRootNodeObject->ParentObject != nullptr;
 
 	// Generate the object expression
-	UE_LOG(LogMutable, Verbose, TEXT("PROFILE: [ %16.8f ] GenerateMutableSource start."), FPlatformTime::Seconds());
+	UE_LOGF(LogMutable, Verbose, "PROFILE: [ %16.8f ] GenerateMutableSource start.", FPlatformTime::Seconds());
 	UE::Mutable::Private::Ptr<UE::Mutable::Private::NodeObject> MutableRoot = GenerateMutableSource(RootNodeObject->OutputPin(), GenerationContext);
-	UE_LOG(LogMutable, Verbose, TEXT("PROFILE: [ %16.8f ] GenerateMutableSource end."), FPlatformTime::Seconds());
+	UE_LOGF(LogMutable, Verbose, "PROFILE: [ %16.8f ] GenerateMutableSource end.", FPlatformTime::Seconds());
 
 	// Generate ReferenceSkeletalMeshes data;
 	PopulateReferenceSkeletalMeshesData(GenerationContext);
@@ -899,7 +899,7 @@ void FCustomizableObjectCompiler::CompileInternal(bool bAsync)
 {
 	MUTABLE_CPUPROFILER_SCOPE(FCustomizableObjectCompiler::CompileInternal)
 	
-	UE_LOG(LogMutable, Verbose, TEXT("PROFILE: [ %16.8f ] FCustomizableObjectCompiler::Compile start."), FPlatformTime::Seconds());
+	UE_LOGF(LogMutable, Verbose, "PROFILE: [ %16.8f ] FCustomizableObjectCompiler::Compile start.", FPlatformTime::Seconds());
 	
 	// This is redundant but necessary to keep static analysis happy.
 	if (!CurrentObject)
@@ -1141,7 +1141,7 @@ void FCustomizableObjectCompiler::CompleteRequest(ECompilationStatePrivate State
 			CurrentObject->GetPrivate()->PostCompile();
 		}
 
-		UE_LOG(LogMutable, Display, TEXT("Finished compiling Customizable Object %s. Compilation took %5.3f seconds to complete."),
+		UE_LOGF(LogMutable, Display, "Finished compiling Customizable Object %ls. Compilation took %5.3f seconds to complete.",
 			*CurrentObject->GetName(), FPlatformTime::Seconds() - CompilationStartTime);
 	}
 
@@ -1205,8 +1205,8 @@ void FCustomizableObjectCompiler::CompleteRequest(ECompilationStatePrivate State
 	PlatformData.Reset();
 	
 
-	UE_LOG(LogMutable, Verbose, TEXT("PROFILE: [ %16.8f ] Completed compile request."), FPlatformTime::Seconds());
-	UE_LOG(LogMutable, Verbose, TEXT("PROFILE: -----------------------------------------------------------"));
+	UE_LOGF(LogMutable, Verbose, "PROFILE: [ %16.8f ] Completed compile request.", FPlatformTime::Seconds());
+	UE_LOGF(LogMutable, Verbose, "PROFILE: -----------------------------------------------------------");
 }
 
 
@@ -1545,10 +1545,10 @@ void FCustomizableObjectCompiler::FinishLoadingStreamableDataFromDDC()
 			CachedPlatformData = MoveTemp(*PlatformData.Get());
 		}
 
-		UE_LOG(LogMutable, Verbose, TEXT("PROFILE: [ %16.8f ] Finishing Compilation task for CO [%s]."), FPlatformTime::Seconds(), *CurrentObject->GetName());
+		UE_LOGF(LogMutable, Verbose, "PROFILE: [ %16.8f ] Finishing Compilation task for CO [%ls].", FPlatformTime::Seconds(), *CurrentObject->GetName());
 		TRACE_END_REGION(UE_MUTABLE_COMPILE_REGION);
 
-		UE_LOG(LogMutable, Display, TEXT("Compiled data loaded from DDC"));
+		UE_LOGF(LogMutable, Display, "Compiled data loaded from DDC");
 		CompleteRequest(ECompilationStatePrivate::Completed, ECompilationResultPrivate::Success);
 	}
 	else
@@ -1564,7 +1564,7 @@ UE::Mutable::Private::Ptr<UE::Mutable::Private::Node> FCustomizableObjectCompile
 	TArray<TSoftObjectPtr<const UStreamableRenderAsset>>& OutRuntimeReferencedMeshes,
 	TArray<FMutableSourceMeshData>& OutCompilerReferencedMeshes)
 {
-	UE_LOG(LogMutable, Log, TEXT("Started Customizable Object Export %s."), *Object->GetName());
+	UE_LOGF(LogMutable, Log, "Started Customizable Object Export %ls.", *Object->GetName());
 
 	FCompilationOptions CompilerOptions = InCompilerOptions;
 	CompilerOptions.bRealTimeMorphTargetsEnabled = Object->GetPrivate()->IsRealTimeMorphTargetsEnabled();
@@ -1674,7 +1674,7 @@ void FCustomizableObjectCompiler::FinishCompilationTask()
 	CompileThread.Reset();
 	CompileTask.Reset();
 
-	UE_LOG(LogMutable, Verbose, TEXT("PROFILE: [ %16.8f ] Finishing Compilation task for CO [%s]."), FPlatformTime::Seconds(), *CurrentObject->GetName());
+	UE_LOGF(LogMutable, Verbose, "PROFILE: [ %16.8f ] Finishing Compilation task for CO [%ls].", FPlatformTime::Seconds(), *CurrentObject->GetName());
 	TRACE_END_REGION(UE_MUTABLE_COMPILE_REGION);
 
 	// Create SaveDD task
@@ -1707,7 +1707,7 @@ void FCustomizableObjectCompiler::FinishSavingDerivedDataTask()
 	SaveDDThread.Reset();
 	SaveDDTask.Reset();
 
-	UE_LOG(LogMutable, Verbose, TEXT("PROFILE: [ %16.8f ] Finished Saving Derived Data task for CO [%s]."), FPlatformTime::Seconds(), *CurrentObject->GetName());
+	UE_LOGF(LogMutable, Verbose, "PROFILE: [ %16.8f ] Finished Saving Derived Data task for CO [%ls].", FPlatformTime::Seconds(), *CurrentObject->GetName());
 
 	TRACE_END_REGION(UE_MUTABLE_SAVEDD_REGION);
 }
@@ -1734,7 +1734,7 @@ void FCustomizableObjectCompiler::ForceFinishCompilation()
 		CompileThread.Reset();
 		CompileTask.Reset();
 
-		UE_LOG(LogMutable, Verbose, TEXT("Force Finish Compilation task for Object."));
+		UE_LOGF(LogMutable, Verbose, "Force Finish Compilation task for Object.");
 		TRACE_END_REGION(UE_MUTABLE_COMPILE_REGION);
 	}
 
@@ -1746,7 +1746,7 @@ void FCustomizableObjectCompiler::ForceFinishCompilation()
 		SaveDDThread.Reset();
 		SaveDDTask.Reset();
 
-		UE_LOG(LogMutable, Verbose, TEXT("Forced Finish Saving Derived Data task."));
+		UE_LOGF(LogMutable, Verbose, "Forced Finish Saving Derived Data task.");
 		TRACE_END_REGION(UE_MUTABLE_SAVEDD_REGION);
 	}
 
diff --git a/Engine/Plugins/Mutable/Source/CustomizableObjectEditor/Private/MuCOE/CustomizableObjectEditorFunctionLibrary.cpp b/Engine/Plugins/Mutable/Source/CustomizableObjectEditor/Private/MuCOE/CustomizableObjectEditorFunctionLibrary.cpp
--- a/Engine/Plugins/Mutable/Source/CustomizableObjectEditor/Private/MuCOE/CustomizableObjectEditorFunctionLibrary.cpp
+++ b/Engine/Plugins/Mutable/Source/CustomizableObjectEditor/Private/MuCOE/CustomizableObjectEditorFunctionLibrary.cpp
@@ -25,7 +25,7 @@ ECustomizableObjectCompilationState UCustomizableObjectEditorFunctionLibrary::Co
 {
 	if (!CustomizableObject)
 	{
-		UE_LOG(LogMutable, Warning, TEXT("Compilation Failed: Trying to compile a null CO!"));
+		UE_LOGF(LogMutable, Warning, "Compilation Failed: Trying to compile a null CO!");
 		return ECustomizableObjectCompilationState::Failed;
 	}
 
@@ -55,7 +55,7 @@ ECustomizableObjectCompilationState UCustomizableObjectEditorFunctionLibrary::Co
 
 	if (!CustomizableObject->IsCompiled())
 	{
-		UE_LOG(LogMutable, Warning, TEXT("CO not marked as compiled"));
+		UE_LOGF(LogMutable, Warning, "CO not marked as compiled");
 	}
 
 	return bCompilationSuccess ? ECustomizableObjectCompilationState::Completed : ECustomizableObjectCompilationState::Failed;
@@ -67,7 +67,7 @@ UCustomizableObject* UCustomizableObjectEditorFunctionLibrary::NewCustomizableOb
 	const FString PackageName = Parameters.PackagePath + "/" + Parameters.AssetName;
 	if (FindPackage(nullptr, *PackageName))
 	{
-		UE_LOG(LogMutable, Error, TEXT("Package [%s] already exists."), *PackageName);
+		UE_LOGF(LogMutable, Error, "Package [%ls] already exists.", *PackageName);
 		return nullptr;
 	}
 
@@ -79,20 +79,20 @@ UCustomizableObject* UCustomizableObjectEditorFunctionLibrary::NewCustomizableOb
 
 	if (OutPackageName != PackageName)
 	{
-		UE_LOG(LogMutable, Error, TEXT("Invalid package name [%s]. Possible valid package name: [%s]"), *PackageName, *OutPackageName);
+		UE_LOGF(LogMutable, Error, "Invalid package name [%ls]. Possible valid package name: [%ls]", *PackageName, *OutPackageName);
 		return nullptr;
 	}
 	
 	if (OutAssetName != Parameters.AssetName)
 	{
-		UE_LOG(LogMutable, Error, TEXT("Invalid asset name [%s]. Possible valid asset name: [%s]"), *Parameters.AssetName, *OutAssetName);
+		UE_LOGF(LogMutable, Error, "Invalid asset name [%ls]. Possible valid asset name: [%ls]", *Parameters.AssetName, *OutAssetName);
 		return nullptr;
 	}
 		
 	UPackage* Package = CreatePackage(*PackageName);
 	if (!Package)
 	{
-		UE_LOG(LogMutable, Error, TEXT("Could not create package [%s]."), *PackageName);
+		UE_LOGF(LogMutable, Error, "Could not create package [%ls].", *PackageName);
 		return nullptr;
 	}
 	
@@ -103,7 +103,7 @@ UCustomizableObject* UCustomizableObjectEditorFunctionLibrary::NewCustomizableOb
 	UObject* Object = AssetToolsModule.Get().CreateAsset(Parameters.AssetName, PackagePath, UCustomizableObject::StaticClass(), Factory);
 	if (!Object)
 	{
-		UE_LOG(LogMutable, Error, TEXT("Could not create Asset [%s]."), *Parameters.AssetName);
+		UE_LOGF(LogMutable, Error, "Could not create Asset [%ls].", *Parameters.AssetName);
 		return nullptr;
 	}
 
diff --git a/Engine/Plugins/Mutable/Source/CustomizableObjectEditor/Private/MuCOE/CustomizableObjectEditorPerformanceAnalyzer.cpp b/Engine/Plugins/Mutable/Source/CustomizableObjectEditor/Private/MuCOE/CustomizableObjectEditorPerformanceAnalyzer.cpp
--- a/Engine/Plugins/Mutable/Source/CustomizableObjectEditor/Private/MuCOE/CustomizableObjectEditorPerformanceAnalyzer.cpp
+++ b/Engine/Plugins/Mutable/Source/CustomizableObjectEditor/Private/MuCOE/CustomizableObjectEditorPerformanceAnalyzer.cpp
@@ -520,7 +520,7 @@ FReply SCustomizableObjectEditorPerformanceAnalyzer::OnInstanceUpdateButtonClick
 		
 	if (!CustomizableObjectBenchmarkingUtils::GenerateDeterministicSetOfInstances(*CustomizableObject ,InstancesToGeneratePerState,InstancesToUpdate,TotalScheduledUpdates))
 	{
-		UE_LOG(LogMutable,Error, TEXT("Mutable Customizable Object Instance generation failed. Aborting perf test."))
+		UE_LOGF(LogMutable,Error, "Mutable Customizable Object Instance generation failed. Aborting perf test.")
 		return FReply::Handled();
 	}
 	check (!InstancesToUpdate.IsEmpty())
diff --git a/Engine/Plugins/Mutable/Source/CustomizableObjectEditor/Private/MuCOE/CustomizableObjectEditorViewportClient.cpp b/Engine/Plugins/Mutable/Source/CustomizableObjectEditor/Private/MuCOE/CustomizableObjectEditorViewportClient.cpp
--- a/Engine/Plugins/Mutable/Source/CustomizableObjectEditor/Private/MuCOE/CustomizableObjectEditorViewportClient.cpp
+++ b/Engine/Plugins/Mutable/Source/CustomizableObjectEditor/Private/MuCOE/CustomizableObjectEditorViewportClient.cpp
@@ -1624,7 +1624,7 @@ void FCustomizableObjectEditorViewportClient::BakeInstance()
 	UCustomizableObjectInstance* Instance = CustomizableObjectEditorPtr.Pin()->GetPreviewInstance();
 	if (!Instance)
 	{
-		UE_LOG(LogMutable, Error, TEXT("No Mutable Customizable Object instance was found in the current editor."));
+		UE_LOGF(LogMutable, Error, "No Mutable Customizable Object instance was found in the current editor.");
 		return;
 	}
 	
@@ -1632,7 +1632,7 @@ void FCustomizableObjectEditorViewportClient::BakeInstance()
 	TObjectPtr<UCustomizableObject> InstanceToBakeCO = Instance->GetCustomizableObject();
 	if (!InstanceToBakeCO)
 	{
-		UE_LOG(LogMutable, Error, TEXT("The Customizable Object instance provided for the baking does not have a Customizable Object set."));
+		UE_LOGF(LogMutable, Error, "The Customizable Object instance provided for the baking does not have a Customizable Object set.");
 		return;
 	}
 
@@ -1653,7 +1653,7 @@ void FCustomizableObjectEditorViewportClient::OnCOForBakingCompilationEnd(const
 {
 	if (InCompileCallbackParams.bRequestFailed || InCompileCallbackParams.bErrors)
 	{
-		UE_LOG(LogMutable, Warning ,TEXT("The Customizable Object failed the compilation. Skipping instance baking"));
+		UE_LOGF(LogMutable, Warning , "The Customizable Object failed the compilation. Skipping instance baking");
 		return;
 	}
 	
@@ -1672,7 +1672,7 @@ void FCustomizableObjectEditorViewportClient::OnInstanceForBakingUpdate(const FU
 	// Early exit if update result is not success
 	if (!UCustomizableObjectSystem::IsUpdateResultValid(Result.UpdateResult))
 	{
-		UE_LOG(LogMutable, Warning ,TEXT("Instance finished update with an error state : %s. Skipping instance baking"), *UEnum::GetValueAsString(Result.UpdateResult));
+		UE_LOGF(LogMutable, Warning , "Instance finished update with an error state : %ls. Skipping instance baking", *UEnum::GetValueAsString(Result.UpdateResult));
 		BakeTempInstance = nullptr;
 		return;
 	}
diff --git a/Engine/Plugins/Mutable/Source/CustomizableObjectEditor/Private/MuCOE/CustomizableObjectInstanceBaker.cpp b/Engine/Plugins/Mutable/Source/CustomizableObjectEditor/Private/MuCOE/CustomizableObjectInstanceBaker.cpp
--- a/Engine/Plugins/Mutable/Source/CustomizableObjectEditor/Private/MuCOE/CustomizableObjectInstanceBaker.cpp
+++ b/Engine/Plugins/Mutable/Source/CustomizableObjectEditor/Private/MuCOE/CustomizableObjectInstanceBaker.cpp
@@ -23,21 +23,21 @@ void UCustomizableObjectInstanceBaker::BakeInstance(UCustomizableObjectInstance*
 
 	if (Configuration.OutputPath.IsEmpty())
 	{
-		UE_LOG(LogMutable,Error,TEXT("No target save path has been provided."))
+		UE_LOGF(LogMutable,Error, "No target save path has been provided.")
 		FinishBakingOperation(false, nullptr);
 		return;
 	}
 	
 	if (!InTargetInstance)
 	{
-		UE_LOG(LogMutable,Error,TEXT("No instance has been provided."))
+		UE_LOGF(LogMutable,Error, "No instance has been provided.")
 		FinishBakingOperation(false, nullptr);
 		return;
 	}
 
 	if (!InTargetInstance->GetCustomizableObject())
 	{
-		UE_LOG(LogMutable,Error,TEXT("The provided instance does not have a Customizable Object defined."))
+		UE_LOGF(LogMutable,Error, "The provided instance does not have a Customizable Object defined.")
 		FinishBakingOperation(false, nullptr);
 		return;
 	}
@@ -64,7 +64,7 @@ void UCustomizableObjectInstanceBaker::OnCompilationEnd(const FCompileCallbackPa
 	if (InCompileCallbackParams.bRequestFailed || InCompileCallbackParams.bErrors)
 	{
 		// Exit if the compilation failed
-		UE_LOG(LogMutable, Display, TEXT("The Customizable Object compilation for baking failed. Aborting instance baking operation."));
+		UE_LOGF(LogMutable, Display, "The Customizable Object compilation for baking failed. Aborting instance baking operation.");
 		FinishBakingOperation(false, nullptr);
 		return;
 	}
@@ -95,7 +95,7 @@ void UCustomizableObjectInstanceBaker::UpdateInstance()
 void UCustomizableObjectInstanceBaker::OnInstanceUpdated(const FUpdateContext& Result) const
 {
 	const EUpdateResult InstanceUpdateResult = Result.UpdateResult;
-	UE_LOG(LogMutable, Display, TEXT("Instance finished update with state : %s."), *UEnum::GetValueAsString(InstanceUpdateResult));
+	UE_LOGF(LogMutable, Display, "Instance finished update with state : %ls.", *UEnum::GetValueAsString(InstanceUpdateResult));
 	
 	if (ensure(InstanceToBake) && InstanceToBake->GetPrivate()->SkeletalMeshStatus == ESkeletalMeshStatus::Success)
 	{
@@ -104,7 +104,7 @@ void UCustomizableObjectInstanceBaker::OnInstanceUpdated(const FUpdateContext& R
 	else
 	{
 		// Report a failed baking operation.
-		UE_LOG(LogMutable, Display, TEXT("Instance updating for baking failed. Aborting instance baking operation."));
+		UE_LOGF(LogMutable, Display, "Instance updating for baking failed. Aborting instance baking operation.");
 		FinishBakingOperation(false, nullptr);
 	}
 }
@@ -136,7 +136,7 @@ void UCustomizableObjectInstanceBaker::Bake() const
 	TArray<FBakedResourceData> SavedPackagesData;
 	{
 		// Display a list with the paths of the assets saved
-		UE_LOG(LogMutable, Log, TEXT("Saved assets (%d) :"), SavedPackages.Num());
+		UE_LOGF(LogMutable, Log, "Saved assets (%d) :", SavedPackages.Num());
 		SavedPackagesData.Reserve(SavedPackages.Num());
 		for (const TPair<UPackage*, const FResourceBakingData>& SavedPackagePair : SavedPackages)
 		{
@@ -156,14 +156,14 @@ void UCustomizableObjectInstanceBaker::Bake() const
 			check(!BakeData.AssetPath.IsEmpty());
 			
 			// Log the path for the caller to know
-			UE_LOG(LogMutable, Log, TEXT("\t%s --- %s"),  *UEnum::GetValueAsString(BakeData.SaveType), *BakeData.AssetPath);
+			UE_LOGF(LogMutable, Log, "\t%ls --- %ls",  *UEnum::GetValueAsString(BakeData.SaveType), *BakeData.AssetPath);
 			
 			SavedPackagesData.Push(BakeData);
 		}
 	}
 
 	// Report that the baking operation has been completed
-	UE_LOG(LogMutable, Display, TEXT("Finishing COI baking procedure."));
+	UE_LOGF(LogMutable, Display, "Finishing COI baking procedure.");
 	FinishBakingOperation(bBakeWasSuccessful, &SavedPackagesData);
 	
 	// Reset the GIsRunningUnattendedScript flag 
diff --git a/Engine/Plugins/Mutable/Source/CustomizableObjectEditor/Private/MuCOE/CustomizableObjectInstanceBakingUtils.cpp b/Engine/Plugins/Mutable/Source/CustomizableObjectEditor/Private/MuCOE/CustomizableObjectInstanceBakingUtils.cpp
--- a/Engine/Plugins/Mutable/Source/CustomizableObjectEditor/Private/MuCOE/CustomizableObjectInstanceBakingUtils.cpp
+++ b/Engine/Plugins/Mutable/Source/CustomizableObjectEditor/Private/MuCOE/CustomizableObjectInstanceBakingUtils.cpp
@@ -231,7 +231,7 @@ void ShowErrorNotification(const FText& InMessage, const FText& InTitle = LOCTEX
 	}
 	else
 	{
-		UE_LOG(LogMutable, Error, TEXT("%s - %s"), *InTitle.ToString(), *InMessage.ToString());
+		UE_LOGF(LogMutable, Error, "%ls - %ls", *InTitle.ToString(), *InMessage.ToString());
 	}
 }
 
@@ -289,7 +289,7 @@ bool ValidateProvidedAssetPath(const FString& FileName, const FString& AssetPath
 {
 	if (AssetPath.IsEmpty())
 	{
-		UE_LOG(LogMutable, Error, TEXT("The AssetPath can not be empty!"));
+		UE_LOGF(LogMutable, Error, "The AssetPath can not be empty!");
 		return false;
 	}
 
@@ -367,7 +367,7 @@ void ScheduleCOCompilationForBaking(UCustomizableObjectInstance& InTargetInstanc
 	// This prevents the queue of updates to have more than one instance for baking.
 	if (bIsBakeOperationAlreadyScheduled)
 	{
-		UE_LOG(LogMutable, Error, TEXT("The CO compilation for baking could not be scheduled. Another instance is being processed at this time."));
+		UE_LOGF(LogMutable, Error, "The CO compilation for baking could not be scheduled. Another instance is being processed at this time.");
 		InCustomizableObjectCompilationDelegate.ExecuteIfBound({true,false,false,false, false});
 		return;
 	}
@@ -375,7 +375,7 @@ void ScheduleCOCompilationForBaking(UCustomizableObjectInstance& InTargetInstanc
 	TObjectPtr<UCustomizableObject> CustomizableObject = InTargetInstance.GetCustomizableObject();
 	if (!CustomizableObject)
 	{
-		UE_LOG(LogMutable, Error, TEXT("The CO compilation for baking could not be scheduled. The provided instance does not have a CO set."));
+		UE_LOGF(LogMutable, Error, "The CO compilation for baking could not be scheduled. The provided instance does not have a CO set.");
 		InCustomizableObjectCompilationDelegate.ExecuteIfBound({true,false,false,false, false});
 		return;
 	}
@@ -525,20 +525,20 @@ bool BakeCustomizableObjectInstance(
 	
 	if (!ValidateProvidedFileName(UserGivenFilename))
 	{
-		UE_LOG(LogMutable, Error, TEXT("The FileName for the instance baking is not valid."));
+		UE_LOGF(LogMutable, Error, "The FileName for the instance baking is not valid.");
 		return false;
 	}
 
 	if (!ValidateProvidedAssetPath(UserGivenFilename,AssetPath,InstanceCO))
 	{
-		UE_LOG(LogMutable, Error, TEXT("The AssetPath for the instance baking is not valid."));
+		UE_LOGF(LogMutable, Error, "The AssetPath for the instance baking is not valid.");
 		return false;
 	}
 	
 	// Exit early if the provided instance does not have a skeletal mesh
 	if (!InInstance.HasAnySkeletalMesh())
 	{
-		UE_LOG(LogMutable, Error, TEXT("The provided instance does not have an skeletal mesh."));
+		UE_LOGF(LogMutable, Error, "The provided instance does not have an skeletal mesh.");
 		return false;
 	}
 
@@ -546,7 +546,7 @@ bool BakeCustomizableObjectInstance(
 	const UModelResources* ModelResources = InstanceCO->GetPrivate()->GetModelResources();
 	if (!ModelResources)
 	{
-		UE_LOG(LogMutable, Error, TEXT("The ModelResources from the Customizable Object is not valid."));
+		UE_LOGF(LogMutable, Error, "The ModelResources from the Customizable Object is not valid.");
 		return false;
 	}
 
@@ -652,7 +652,7 @@ bool BakeCustomizableObjectInstance(
 										UTexture* DupTexture = Cast<UTexture>(DuplicatedObject);
 										TextureReplacementMaps[MaterialIndex].Add(ParameterInfoIndex, DupTexture);
 
-										UE_LOG(LogMutable, Verbose, TEXT("Matching texture found with name %s in component %s"), *TextureData.GetResourceComposedName(), *ComponentName.ToString())
+										UE_LOGF(LogMutable, Verbose, "Matching texture found with name %ls in component %ls", *TextureData.GetResourceComposedName(), *ComponentName.ToString())
 										break;
 									}
 								}
@@ -899,7 +899,7 @@ bool BakeCustomizableObjectInstance(
 									}
 									else
 									{
-										UE_LOG(LogMutable, Error, TEXT("A Mutable texture that is not a Texture2D has been found while baking a CustomizableObjectInstance."));
+										UE_LOGF(LogMutable, Error, "A Mutable texture that is not a Texture2D has been found while baking a CustomizableObjectInstance.");
 									}
 								}
 								else
@@ -967,7 +967,7 @@ bool BakeCustomizableObjectInstance(
 						{
 							ReplacementMap.Add(Skeleton, DuplicatedObject);
 
-							UE_LOG(LogMutable, Verbose, TEXT("Matching skeleton found with name %s in component %s"), *SkeletonData.GetResourceComposedName(), *ComponentName.ToString())
+							UE_LOGF(LogMutable, Verbose, "Matching skeleton found with name %ls in component %ls", *SkeletonData.GetResourceComposedName(), *ComponentName.ToString())
 							break;
 						}
 					}
@@ -1008,7 +1008,7 @@ bool BakeCustomizableObjectInstance(
 						{
 							ReplacementMap.Add(PhysicsAsset, DuplicatedObject);
 
-							UE_LOG(LogMutable, Verbose, TEXT("Matching Physics Asset found with name %s in component %s"), *PhysicsAssetData.GetResourceComposedName(), *ComponentName.ToString())
+							UE_LOGF(LogMutable, Verbose, "Matching Physics Asset found with name %ls in component %ls", *PhysicsAssetData.GetResourceComposedName(), *ComponentName.ToString())
 							break;
 						}
 					}
diff --git a/Engine/Plugins/Mutable/Source/CustomizableObjectEditor/Private/MuCOE/CustomizableObjectMacroLibrary/CustomizableObjectMacroLibraryEditor.cpp b/Engine/Plugins/Mutable/Source/CustomizableObjectEditor/Private/MuCOE/CustomizableObjectMacroLibrary/CustomizableObjectMacroLibraryEditor.cpp
--- a/Engine/Plugins/Mutable/Source/CustomizableObjectEditor/Private/MuCOE/CustomizableObjectMacroLibrary/CustomizableObjectMacroLibraryEditor.cpp
+++ b/Engine/Plugins/Mutable/Source/CustomizableObjectEditor/Private/MuCOE/CustomizableObjectMacroLibrary/CustomizableObjectMacroLibraryEditor.cpp
@@ -260,7 +260,7 @@ void FCustomizableObjectMacroLibraryEditor::PostUndo(bool bSuccess)
 	{
 		//TODO(Max): Select the correct macro after an undo.
 		/*FTransactionContext TransactionContext = GEditor->GetTransactionName();
-		UE_LOG(LogTemp, Warning, TEXT("Transaction: %s"), *TransactionContext.Title.ToString());*/
+		UE_LOGF(LogTemp, Warning, "Transaction: %ls", *TransactionContext.Title.ToString());*/
 
 		MacroSelector->GenerateRowView();
 	}
diff --git a/Engine/Plugins/Mutable/Source/CustomizableObjectEditor/Private/MuCOE/EdGraphSchema_CustomizableObject.cpp b/Engine/Plugins/Mutable/Source/CustomizableObjectEditor/Private/MuCOE/EdGraphSchema_CustomizableObject.cpp
--- a/Engine/Plugins/Mutable/Source/CustomizableObjectEditor/Private/MuCOE/EdGraphSchema_CustomizableObject.cpp
+++ b/Engine/Plugins/Mutable/Source/CustomizableObjectEditor/Private/MuCOE/EdGraphSchema_CustomizableObject.cpp
@@ -1373,7 +1373,7 @@ void UEdGraphSchema_CustomizableObject::DroppedAssetsOnGraph(const TArray<FAsset
 		// Error : a switch entry is missing for a ESpawnableObjectType value
 		default:
 			{
-				UE_LOG(LogTemp,Error,TEXT("Unable to create new mutable node for target asset : Invalid ESpawnableObjectType value. "));
+				UE_LOGF(LogTemp,Error, "Unable to create new mutable node for target asset : Invalid ESpawnableObjectType value. ");
 				checkNoEntry();
 				break;
 			}
@@ -1393,7 +1393,7 @@ void UEdGraphSchema_CustomizableObject::DroppedAssetsOnGraph(const TArray<FAsset
 		}
 		else
 		{
-			UE_LOG(LogTemp,Error,TEXT("Unable to add null node to graph. "))
+			UE_LOGF(LogTemp,Error, "Unable to add null node to graph. ")
 		}
 	}
 }
diff --git a/Engine/Plugins/Mutable/Source/CustomizableObjectEditor/Private/MuCOE/GenerateMutableSource/GenerateMutableSource.cpp b/Engine/Plugins/Mutable/Source/CustomizableObjectEditor/Private/MuCOE/GenerateMutableSource/GenerateMutableSource.cpp
--- a/Engine/Plugins/Mutable/Source/CustomizableObjectEditor/Private/MuCOE/GenerateMutableSource/GenerateMutableSource.cpp
+++ b/Engine/Plugins/Mutable/Source/CustomizableObjectEditor/Private/MuCOE/GenerateMutableSource/GenerateMutableSource.cpp
@@ -832,7 +832,7 @@ uint32 AddUniqueSurfaceMetadata(const FMutableSourceSurfaceMetadata& SourceMetad
 	}
 	else
 	{
-		UE_LOG(LogMutable, Error, TEXT("Maximum number of surfaces reached."));
+		UE_LOGF(LogMutable, Error, "Maximum number of surfaces reached.");
 	}
 
 	return SurfaceMetadataUniqueHash;
diff --git a/Engine/Plugins/Mutable/Source/CustomizableObjectEditor/Private/MuCOE/GenerateMutableSource/GenerateMutableSourceImage.cpp b/Engine/Plugins/Mutable/Source/CustomizableObjectEditor/Private/MuCOE/GenerateMutableSource/GenerateMutableSourceImage.cpp
--- a/Engine/Plugins/Mutable/Source/CustomizableObjectEditor/Private/MuCOE/GenerateMutableSource/GenerateMutableSourceImage.cpp
+++ b/Engine/Plugins/Mutable/Source/CustomizableObjectEditor/Private/MuCOE/GenerateMutableSource/GenerateMutableSourceImage.cpp
@@ -1377,7 +1377,7 @@ UE::Mutable::Private::Ptr<UE::Mutable::Private::NodeImageFormat> GenerateMutable
 					*ReferenceTexture->GetName());
 				const FText ReplacedImageFormatText = FText::FromString(ReplacedImageFormatMsg);
 				GenerationContext.Log(ReplacedImageFormatText, Node, EMessageSeverity::Info);
-				UE_LOG(LogMutable, Log, TEXT("%s"), *ReplacedImageFormatMsg);
+				UE_LOGF(LogMutable, Log, "%ls", *ReplacedImageFormatMsg);
 			}
 			else if (BuildSettingsPerFormatPerLayer.Num() > 1)
 			{
@@ -1387,7 +1387,7 @@ UE::Mutable::Private::Ptr<UE::Mutable::Private::NodeImageFormat> GenerateMutable
 					*ReferenceTexture->GetName());
 				const FText ReplacedImageFormatText = FText::FromString(ReplacedImageFormatMsg);
 				GenerationContext.Log(ReplacedImageFormatText, Node, EMessageSeverity::Info);
-				UE_LOG(LogMutable, Log, TEXT("%s"), *ReplacedImageFormatMsg);
+				UE_LOGF(LogMutable, Log, "%ls", *ReplacedImageFormatMsg);
 			}
 		}
 
@@ -1460,7 +1460,7 @@ UE::Mutable::Private::Ptr<UE::Mutable::Private::NodeImageFormat> GenerateMutable
 					const FString ReplacedImageFormatMsg = FString::Printf(TEXT("In object [%s] the unsupported image format %d is used, RGBA_UByte will be used instead."), *GenerationContext.GetObjectName(), UnrealTargetPlatformFormat);
 					const FText ReplacedImageFormatText = FText::FromString(ReplacedImageFormatMsg);
 					GenerationContext.Log(ReplacedImageFormatText, Node, EMessageSeverity::Info);
-					UE_LOG(LogMutable, Log, TEXT("%s"), *ReplacedImageFormatMsg);
+					UE_LOGF(LogMutable, Log, "%ls", *ReplacedImageFormatMsg);
 					MutableFormat = UE::Mutable::Private::EImageFormat::RGBA_UByte;
 				}
 				if (MutableFormatIfAlpha == UE::Mutable::Private::EImageFormat::None)
@@ -1468,7 +1468,7 @@ UE::Mutable::Private::Ptr<UE::Mutable::Private::NodeImageFormat> GenerateMutable
 					const FString ReplacedImageFormatMsg = FString::Printf(TEXT("In object [%s] the unsupported image format %d is used, RGBA_UByte will be used instead."), *GenerationContext.GetObjectName(), UnrealTargetPlatformFormatAlpha);
 					const FText ReplacedImageFormatText = FText::FromString(ReplacedImageFormatMsg);
 					GenerationContext.Log(ReplacedImageFormatText, Node, EMessageSeverity::Info);
-					UE_LOG(LogMutable, Log, TEXT("%s"), *ReplacedImageFormatMsg);
+					UE_LOGF(LogMutable, Log, "%ls", *ReplacedImageFormatMsg);
 					MutableFormatIfAlpha = UE::Mutable::Private::EImageFormat::RGBA_UByte;
 				}
 
diff --git a/Engine/Plugins/Mutable/Source/CustomizableObjectEditor/Private/MuCOE/GenerateMutableSource/GenerateMutableSourceLayout.cpp b/Engine/Plugins/Mutable/Source/CustomizableObjectEditor/Private/MuCOE/GenerateMutableSource/GenerateMutableSourceLayout.cpp
--- a/Engine/Plugins/Mutable/Source/CustomizableObjectEditor/Private/MuCOE/GenerateMutableSource/GenerateMutableSourceLayout.cpp
+++ b/Engine/Plugins/Mutable/Source/CustomizableObjectEditor/Private/MuCOE/GenerateMutableSource/GenerateMutableSourceLayout.cpp
@@ -92,7 +92,7 @@ UE::Mutable::Private::FSourceLayoutBlock ToMutable(const FCustomizableObjectLayo
 		if (Error != EUnrealToMutableConversionError::Success)
 		{
 			// This should never happen, so details are not necessary.
-			UE_LOG(LogMutable, Warning, TEXT("Failed to convert layout block mask texture."));
+			UE_LOGF(LogMutable, Warning, "Failed to convert layout block mask texture.");
 		}
 		else
 		{
diff --git a/Engine/Plugins/Mutable/Source/CustomizableObjectEditor/Private/MuCOE/GenerateMutableSource/GenerateMutableSourceMesh.cpp b/Engine/Plugins/Mutable/Source/CustomizableObjectEditor/Private/MuCOE/GenerateMutableSource/GenerateMutableSourceMesh.cpp
--- a/Engine/Plugins/Mutable/Source/CustomizableObjectEditor/Private/MuCOE/GenerateMutableSource/GenerateMutableSourceMesh.cpp
+++ b/Engine/Plugins/Mutable/Source/CustomizableObjectEditor/Private/MuCOE/GenerateMutableSource/GenerateMutableSourceMesh.cpp
@@ -136,7 +136,7 @@ void SetBuffersFormat(FMutableCompilationContext& Context,
 
 	if (MutableBonesPerVertex != MaxNumBonesPerVertex)
 	{
-		UE_LOG(LogMutable, Verbose, TEXT("In object [%s] Mesh bone number adjusted from %d to %d."), *Context.Object->GetFullName(), MaxNumBonesPerVertex, MutableBonesPerVertex);
+		UE_LOGF(LogMutable, Verbose, "In object [%ls] Mesh bone number adjusted from %d to %d.", *Context.Object->GetFullName(), MaxNumBonesPerVertex, MutableBonesPerVertex);
 	}
 	
 	const bool bHasTexCoords = !bIgnoreTexCoords;
@@ -1346,7 +1346,7 @@ UE::Tasks::FTask ConvertSkeletalMeshToMutable(TSharedRef<FMeshConversionContext>
 
 			if (ImportedProfileData->SkinWeights.Num() != LODModel.NumVertices)
 			{
-				UE_LOG(LogMutable, Display, TEXT("Invalid data found in SkeletalMesh [%s] when converting SkinWeightProfiles. Regenerate and save the SkeletalMesh."), *SkeletalMesh->GetName());
+				UE_LOGF(LogMutable, Display, "Invalid data found in SkeletalMesh [%ls] when converting SkinWeightProfiles. Regenerate and save the SkeletalMesh.", *SkeletalMesh->GetName());
 				
 				FMemory::Memset(DestData, 0, DestDataSize * VertexCount);
 				continue;
@@ -1502,7 +1502,7 @@ UE::Tasks::FTask ConvertSkeletalMeshToMutable(TSharedRef<FMeshConversionContext>
 		}
 		else
 		{
-			UE_LOG(LogMutable, Error, TEXT("Maximum number of meshes reached."));
+			UE_LOGF(LogMutable, Error, "Maximum number of meshes reached.");
 		}
 
 
@@ -1521,7 +1521,7 @@ UE::Tasks::FTask ConvertSkeletalMeshToMutable(TSharedRef<FMeshConversionContext>
 
 			if (MeshSection.ClothMappingDataLODs.Num() > 1)
 			{
-				UE_LOG(LogMutable, Error, TEXT("\"Mappings to Same LOD\" only supported in Ray Tracing \"Cloth LODBias Mode\" option for [%s] Skeletal Mesh"), *SkeletalMesh->GetName());	
+				UE_LOGF(LogMutable, Error, "\"Mappings to Same LOD\" only supported in Ray Tracing \"Cloth LODBias Mode\" option for [%ls] Skeletal Mesh", *SkeletalMesh->GetName());	
 			}
 			
 			const uint32 NumClothVertices = MeshSection.ClothMappingDataLODs[0].Num();
diff --git a/Engine/Plugins/Mutable/Source/CustomizableObjectEditor/Private/MuCOE/MutableMeshPreviewUtils.cpp b/Engine/Plugins/Mutable/Source/CustomizableObjectEditor/Private/MuCOE/MutableMeshPreviewUtils.cpp
--- a/Engine/Plugins/Mutable/Source/CustomizableObjectEditor/Private/MuCOE/MutableMeshPreviewUtils.cpp
+++ b/Engine/Plugins/Mutable/Source/CustomizableObjectEditor/Private/MuCOE/MutableMeshPreviewUtils.cpp
@@ -209,7 +209,7 @@ namespace MutableMeshPreviewUtils
 			{
 				if (bLogFindingsToConsole)
 				{
-					UE_LOG(LogMutable, Warning, TEXT("\tVertex Positions Buffer not found."));
+					UE_LOGF(LogMutable, Warning, "\tVertex Positions Buffer not found.");
 				}
 
 				bCanBeConverted = false;
@@ -233,7 +233,7 @@ namespace MutableMeshPreviewUtils
 			{
 				if (bLogFindingsToConsole)
 				{
-					UE_LOG(LogMutable, Warning, TEXT("\tSkinning Buffer not found."));
+					UE_LOGF(LogMutable, Warning, "\tSkinning Buffer not found.");
 				}
 
 				bCanBeConverted = false;
@@ -245,7 +245,7 @@ namespace MutableMeshPreviewUtils
 			{
 				if (bLogFindingsToConsole)
 				{
-					UE_LOG(LogMutable, Warning, TEXT("\tTangents Buffer not found."));
+					UE_LOGF(LogMutable, Warning, "\tTangents Buffer not found.");
 				}
 
 				bCanBeConverted = false;
@@ -269,7 +269,7 @@ namespace MutableMeshPreviewUtils
 			{
 				if (bLogFindingsToConsole)
 				{
-					UE_LOG(LogMutable, Warning, TEXT("\tTexCords Buffer not found."));
+					UE_LOGF(LogMutable, Warning, "\tTexCords Buffer not found.");
 				}
 
 				bCanBeConverted = false;
@@ -487,11 +487,11 @@ namespace MutableMeshPreviewUtils
 				if (!AreVertexBuffersCorrectlyLocated(InMutableMesh,true))
 				{
 					// The formatting operation failed, do not proceed
-					UE_LOG(LogMutable,Error,TEXT("Restructuring of mutable mesh buffers failed: Aborting Skeletal Mesh generation. "));
+					UE_LOGF(LogMutable,Error, "Restructuring of mutable mesh buffers failed: Aborting Skeletal Mesh generation. ");
 					return nullptr;
 				}
 				
-				UE_LOG(LogMutable, Verbose, TEXT("Buffers Restructuring was succesfull."));
+				UE_LOGF(LogMutable, Verbose, "Buffers Restructuring was succesfull.");
 			}
 
 			// We should now have a mutable mesh following a buffer structure where the required buffers are there to be
diff --git a/Engine/Plugins/Mutable/Source/CustomizableObjectEditor/Private/MuCOE/Nodes/CONodeModifierEditSkeletalMeshSection.cpp b/Engine/Plugins/Mutable/Source/CustomizableObjectEditor/Private/MuCOE/Nodes/CONodeModifierEditSkeletalMeshSection.cpp
--- a/Engine/Plugins/Mutable/Source/CustomizableObjectEditor/Private/MuCOE/Nodes/CONodeModifierEditSkeletalMeshSection.cpp
+++ b/Engine/Plugins/Mutable/Source/CustomizableObjectEditor/Private/MuCOE/Nodes/CONodeModifierEditSkeletalMeshSection.cpp
@@ -173,7 +173,7 @@ void UCONodeModifierEditSkeletalMeshSection::BackwardsCompatibleFixup(int32 Cust
 
 		if (!Layouts.IsValidIndex(ParentLayoutIndex))
 		{
-			UE_LOG(LogMutable, Warning, TEXT("[%s] UCONodeModifierEditSkeletalMeshSection refers to an invalid texture layout index %d. Parent node has %d layouts."), 
+			UE_LOGF(LogMutable, Warning, "[%ls] UCONodeModifierEditSkeletalMeshSection refers to an invalid texture layout index %d. Parent node has %d layouts.", 
 				*GetOutermost()->GetName(), ParentLayoutIndex, Layouts.Num());
 		}
 		else
@@ -194,13 +194,13 @@ void UCONodeModifierEditSkeletalMeshSection::BackwardsCompatibleFixup(int32 Cust
 						}
 						else
 						{
-							UE_LOG(LogMutable, Warning, TEXT("[%s] UCONodeModifierEditSkeletalMeshSection refers to an valid layout block %d but that block doesn't have an id."),
+							UE_LOGF(LogMutable, Warning, "[%ls] UCONodeModifierEditSkeletalMeshSection refers to an valid layout block %d but that block doesn't have an id.",
 								*GetOutermost()->GetName(), BlockIndex );
 						}
 					}
 					else
 					{
-						UE_LOG(LogMutable, Warning, TEXT("[%s] UCONodeModifierEditSkeletalMeshSection refers to an invalid layout block index %d. Parent node has %d blocks."),
+						UE_LOGF(LogMutable, Warning, "[%ls] UCONodeModifierEditSkeletalMeshSection refers to an invalid layout block index %d. Parent node has %d blocks.",
 							*GetOutermost()->GetName(), BlockIndex, ParentLayout->Blocks.Num());
 					}
 				}
@@ -221,7 +221,7 @@ void UCONodeModifierEditSkeletalMeshSection::BackwardsCompatibleFixup(int32 Cust
 
 		if (!ParentLayouts.IsValidIndex(ParentLayoutIndex))
 		{
-			UE_LOG(LogMutable, Warning, TEXT("[%s] UCONodeModifierEditSkeletalMeshSection refers to an invalid texture layout index %d. Parent node has %d layouts."),
+			UE_LOGF(LogMutable, Warning, "[%ls] UCONodeModifierEditSkeletalMeshSection refers to an invalid texture layout index %d. Parent node has %d layouts.",
 				*GetOutermost()->GetName(), ParentLayoutIndex, ParentLayouts.Num());
 		}
 		else
@@ -241,7 +241,7 @@ void UCONodeModifierEditSkeletalMeshSection::BackwardsCompatibleFixup(int32 Cust
 						if (ExistingBlock.Id == BlockId)
 						{
 							bSkipBlock = true;
-							UE_LOG(LogMutable, Log, TEXT("[%s] UCONodeModifierEditSkeletalMeshSection has a duplicated layout block id. One has been ignored during version upgrade."), *GetOutermost()->GetName());
+							UE_LOGF(LogMutable, Log, "[%ls] UCONodeModifierEditSkeletalMeshSection has a duplicated layout block id. One has been ignored during version upgrade.", *GetOutermost()->GetName());
 							break;
 						}
 					}
@@ -273,7 +273,7 @@ void UCONodeModifierEditSkeletalMeshSection::BackwardsCompatibleFixup(int32 Cust
 
 					if (!bFoundInParent)
 					{
-						UE_LOG(LogMutable, Warning, TEXT("[%s] UCONodeModifierEditSkeletalMeshSection refers to and invalid layout block. It has been ignored during version upgrade."), *GetOutermost()->GetName());
+						UE_LOGF(LogMutable, Warning, "[%ls] UCONodeModifierEditSkeletalMeshSection refers to and invalid layout block. It has been ignored during version upgrade.", *GetOutermost()->GetName());
 					}
 				}
 			}
@@ -350,7 +350,7 @@ void UCONodeModifierEditSkeletalMeshSection::BackwardsCompatibleFixup(int32 Cust
 		{
 			// Conversion failed?
 			ensure(false);
-			UE_LOG(LogMutable, Warning, TEXT("[%s] UCONodeModifierExtendSkeletalMeshSection version upgrade failed."), *GetOutermost()->GetName());
+			UE_LOGF(LogMutable, Warning, "[%ls] UCONodeModifierExtendSkeletalMeshSection version upgrade failed.", *GetOutermost()->GetName());
 		}
 
 		ReconstructNode();
diff --git a/Engine/Plugins/Mutable/Source/CustomizableObjectEditor/Private/MuCOE/Nodes/CONodeModifierExtendSkeletalMeshSection.cpp b/Engine/Plugins/Mutable/Source/CustomizableObjectEditor/Private/MuCOE/Nodes/CONodeModifierExtendSkeletalMeshSection.cpp
--- a/Engine/Plugins/Mutable/Source/CustomizableObjectEditor/Private/MuCOE/Nodes/CONodeModifierExtendSkeletalMeshSection.cpp
+++ b/Engine/Plugins/Mutable/Source/CustomizableObjectEditor/Private/MuCOE/Nodes/CONodeModifierExtendSkeletalMeshSection.cpp
@@ -88,7 +88,7 @@ void UCONodeModifierExtendSkeletalMeshSection::BackwardsCompatibleFixup(int32 Cu
 			{
 				// Conversion failed?
 				ensure(false);
-				UE_LOG(LogMutable, Warning, TEXT("[%s] UCONodeModifierExtendSkeletalMeshSection version upgrade failed."), *GetOutermost()->GetName());
+				UE_LOGF(LogMutable, Warning, "[%ls] UCONodeModifierExtendSkeletalMeshSection version upgrade failed.", *GetOutermost()->GetName());
 			}
 
 			// If the tag was added to another CO, keep track of the addition to be able to repeat it in case
@@ -106,7 +106,7 @@ void UCONodeModifierExtendSkeletalMeshSection::BackwardsCompatibleFixup(int32 Cu
 		else
 		{
 			// No parent, no conversion.
-			UE_LOG(LogMutable, Log, TEXT("[%s] UCONodeModifierExtendSkeletalMeshSection has no parent. It will not be upgraded."), *GetOutermost()->GetName());
+			UE_LOGF(LogMutable, Log, "[%ls] UCONodeModifierExtendSkeletalMeshSection has no parent. It will not be upgraded.", *GetOutermost()->GetName());
 		}
 
 		ReconstructNode();
diff --git a/Engine/Plugins/Mutable/Source/CustomizableObjectEditor/Private/MuCOE/Nodes/CustomizableObjectNodeGroupProjectorParameter.cpp b/Engine/Plugins/Mutable/Source/CustomizableObjectEditor/Private/MuCOE/Nodes/CustomizableObjectNodeGroupProjectorParameter.cpp
--- a/Engine/Plugins/Mutable/Source/CustomizableObjectEditor/Private/MuCOE/Nodes/CustomizableObjectNodeGroupProjectorParameter.cpp
+++ b/Engine/Plugins/Mutable/Source/CustomizableObjectEditor/Private/MuCOE/Nodes/CustomizableObjectNodeGroupProjectorParameter.cpp
@@ -76,7 +76,7 @@ TArray<FGroupProjectorParameterImage> UCustomizableObjectNodeGroupProjectorParam
 
 	if (PropertyTexturePath == nullptr)
 	{
-		UE_LOG(LogMutable, Warning, TEXT("WARNING: No column found with texture path information to load projection textures"));
+		UE_LOGF(LogMutable, Warning, "WARNING: No column found with texture path information to load projection textures");
 		return ArrayResult;
 	}
 
@@ -91,7 +91,7 @@ TArray<FGroupProjectorParameterImage> UCustomizableObjectNodeGroupProjectorParam
 
 		if (Texture == nullptr)
 		{
-			UE_LOG(LogMutable, Warning, TEXT("WARNING: Unable to load texture %s"), *PropertyValue);
+			UE_LOGF(LogMutable, Warning, "WARNING: Unable to load texture %ls", *PropertyValue);
 		}
 		else
 		{
diff --git a/Engine/Plugins/Mutable/Source/CustomizableObjectEditor/Private/MuCOE/Nodes/CustomizableObjectNodeMeshClipWithMesh.cpp b/Engine/Plugins/Mutable/Source/CustomizableObjectEditor/Private/MuCOE/Nodes/CustomizableObjectNodeMeshClipWithMesh.cpp
--- a/Engine/Plugins/Mutable/Source/CustomizableObjectEditor/Private/MuCOE/Nodes/CustomizableObjectNodeMeshClipWithMesh.cpp
+++ b/Engine/Plugins/Mutable/Source/CustomizableObjectEditor/Private/MuCOE/Nodes/CustomizableObjectNodeMeshClipWithMesh.cpp
@@ -63,7 +63,7 @@ void UCustomizableObjectNodeModifierClipWithMesh::BackwardsCompatibleFixup(int32
 				{
 					// Conversion failed?
 					ensure(false);
-					UE_LOG(LogMutable, Warning, TEXT("[%s] UCustomizableObjectNodeModifierClipWithMesh version upgrade failed."), *GetOutermost()->GetName());
+					UE_LOGF(LogMutable, Warning, "[%ls] UCustomizableObjectNodeModifierClipWithMesh version upgrade failed.", *GetOutermost()->GetName());
 				}
 
 				// If the tag was added to another CO, keep track of the addition to be able to repeat it in case
@@ -80,7 +80,7 @@ void UCustomizableObjectNodeModifierClipWithMesh::BackwardsCompatibleFixup(int32
 			}
 			else
 			{
-				UE_LOG(LogMutable, Log, TEXT("[%s] UCustomizableObjectNodeModifierClipWithMesh has no parent. It will not be upgraded."), *GetOutermost()->GetName());
+				UE_LOGF(LogMutable, Log, "[%ls] UCustomizableObjectNodeModifierClipWithMesh has no parent. It will not be upgraded.", *GetOutermost()->GetName());
 			}
 		}
 
diff --git a/Engine/Plugins/Mutable/Source/CustomizableObjectEditor/Private/MuCOE/Nodes/CustomizableObjectNodeModifierEditMeshSectionBase.cpp b/Engine/Plugins/Mutable/Source/CustomizableObjectEditor/Private/MuCOE/Nodes/CustomizableObjectNodeModifierEditMeshSectionBase.cpp
--- a/Engine/Plugins/Mutable/Source/CustomizableObjectEditor/Private/MuCOE/Nodes/CustomizableObjectNodeModifierEditMeshSectionBase.cpp
+++ b/Engine/Plugins/Mutable/Source/CustomizableObjectEditor/Private/MuCOE/Nodes/CustomizableObjectNodeModifierEditMeshSectionBase.cpp
@@ -42,7 +42,7 @@ void UCustomizableObjectNodeModifierEditMeshSectionBase::BackwardsCompatibleFixu
 			{
 				// Conversion failed?
 				ensure(false);
-				UE_LOG(LogMutable, Warning, TEXT("[%s] UCustomizableObjectNodeModifierEditMeshSectionBase version upgrade failed."), *GetOutermost()->GetName());
+				UE_LOGF(LogMutable, Warning, "[%ls] UCustomizableObjectNodeModifierEditMeshSectionBase version upgrade failed.", *GetOutermost()->GetName());
 			}
 
 			// If the tag was added to another CO, keep track of the addition to be able to repeat it in case
@@ -59,7 +59,7 @@ void UCustomizableObjectNodeModifierEditMeshSectionBase::BackwardsCompatibleFixu
 		}
 		else
 		{
-			UE_LOG(LogMutable, Log, TEXT("[%s] UCustomizableObjectNodeModifierEditMeshSectionBase has no parent. It will not be upgraded."), *GetOutermost()->GetName());
+			UE_LOGF(LogMutable, Log, "[%ls] UCustomizableObjectNodeModifierEditMeshSectionBase has no parent. It will not be upgraded.", *GetOutermost()->GetName());
 		}
 	}
 }
diff --git a/Engine/Plugins/Mutable/Source/CustomizableObjectEditor/Private/MuCOE/Nodes/CustomizableObjectNodeModifierRemoveMeshBlocks.cpp b/Engine/Plugins/Mutable/Source/CustomizableObjectEditor/Private/MuCOE/Nodes/CustomizableObjectNodeModifierRemoveMeshBlocks.cpp
--- a/Engine/Plugins/Mutable/Source/CustomizableObjectEditor/Private/MuCOE/Nodes/CustomizableObjectNodeModifierRemoveMeshBlocks.cpp
+++ b/Engine/Plugins/Mutable/Source/CustomizableObjectEditor/Private/MuCOE/Nodes/CustomizableObjectNodeModifierRemoveMeshBlocks.cpp
@@ -35,7 +35,7 @@ void UCustomizableObjectNodeModifierRemoveMeshBlocks::BackwardsCompatibleFixup(i
 
 			if (!Layouts.IsValidIndex(ParentLayoutIndex))
 			{
-				UE_LOG(LogMutable, Warning, TEXT("[%s] UCustomizableObjectNodeModifierRemoveMeshBlocks refers to an invalid texture layout index %d. Parent node has %d layouts."),
+				UE_LOGF(LogMutable, Warning, "[%ls] UCustomizableObjectNodeModifierRemoveMeshBlocks refers to an invalid texture layout index %d. Parent node has %d layouts.",
 					*GetOutermost()->GetName(), ParentLayoutIndex, Layouts.Num());
 			}
 			else if (UCONodeSkeletalMeshSection* ParentMaterial = Cast<UCONodeSkeletalMeshSection>(ParentMaterialNode))
@@ -47,7 +47,7 @@ void UCustomizableObjectNodeModifierRemoveMeshBlocks::BackwardsCompatibleFixup(i
 					const int BlockIndex = Blocks_DEPRECATED[IndexIndex];
 					if (!ParentLayout->Blocks.IsValidIndex(BlockIndex))
 					{
-						UE_LOG(LogMutable, Warning, TEXT("[%s] UCustomizableObjectNodeModifierRemoveMeshBlocks refers to an invalid layout block index %d. Parent node has %d blocks."),
+						UE_LOGF(LogMutable, Warning, "[%ls] UCustomizableObjectNodeModifierRemoveMeshBlocks refers to an invalid layout block index %d. Parent node has %d blocks.",
 							*GetOutermost()->GetName(), BlockIndex, ParentLayout->Blocks.Num());
 
 						continue;
@@ -56,7 +56,7 @@ void UCustomizableObjectNodeModifierRemoveMeshBlocks::BackwardsCompatibleFixup(i
 					const FGuid Id = ParentLayout->Blocks[BlockIndex].Id;
 					if (!Id.IsValid())
 					{
-						UE_LOG(LogMutable, Warning, TEXT("[%s] UCustomizableObjectNodeModifierRemoveMeshBlocks refers to an valid layout block %d but that block doesn't have an id."),
+						UE_LOGF(LogMutable, Warning, "[%ls] UCustomizableObjectNodeModifierRemoveMeshBlocks refers to an valid layout block %d but that block doesn't have an id.",
 							*GetOutermost()->GetName(), BlockIndex);
 
 						continue;
@@ -68,7 +68,7 @@ void UCustomizableObjectNodeModifierRemoveMeshBlocks::BackwardsCompatibleFixup(i
 		}
 		else
 		{
-			UE_LOG(LogMutable, Log, TEXT("[%s] UCustomizableObjectNodeModifierRemoveMeshBlocks has no parent. It will not be upgraded."), *GetOutermost()->GetName());
+			UE_LOGF(LogMutable, Log, "[%ls] UCustomizableObjectNodeModifierRemoveMeshBlocks has no parent. It will not be upgraded.", *GetOutermost()->GetName());
 		}
 	}
 
@@ -86,7 +86,7 @@ void UCustomizableObjectNodeModifierRemoveMeshBlocks::BackwardsCompatibleFixup(i
 		{
 			// Conversion failed?
 			ensure(false);
-			UE_LOG(LogMutable, Warning, TEXT("[%s] UCustomizableObjectNodeModifierRemoveMeshBlocks version upgrade failed."), *GetOutermost()->GetName());
+			UE_LOGF(LogMutable, Warning, "[%ls] UCustomizableObjectNodeModifierRemoveMeshBlocks version upgrade failed.", *GetOutermost()->GetName());
 		}
 		else
 		{
@@ -160,7 +160,7 @@ void UCustomizableObjectNodeModifierRemoveMeshBlocks::BackwardsCompatibleFixup(i
 
 		if (!ParentLayouts.IsValidIndex(ParentLayoutIndex))
 		{
-			UE_LOG(LogMutable, Warning, TEXT("[%s] UCustomizableObjectNodeModifierRemoveMeshBlocks refers to an invalid texture layout index %d. Parent node has %d layouts."),
+			UE_LOGF(LogMutable, Warning, "[%ls] UCustomizableObjectNodeModifierRemoveMeshBlocks refers to an invalid texture layout index %d. Parent node has %d layouts.",
 				*GetOutermost()->GetName(), ParentLayoutIndex, ParentLayouts.Num());
 		}
 		else
@@ -180,7 +180,7 @@ void UCustomizableObjectNodeModifierRemoveMeshBlocks::BackwardsCompatibleFixup(i
 						if (ExistingBlock.Id == BlockId)
 						{
 							bSkipBlock = true;
-							UE_LOG(LogMutable, Log, TEXT("[%s] UCustomizableObjectNodeModifierRemoveMeshBlocks has a duplicated layout block id. One has been ignored during version upgrade."), *GetOutermost()->GetName());
+							UE_LOGF(LogMutable, Log, "[%ls] UCustomizableObjectNodeModifierRemoveMeshBlocks has a duplicated layout block id. One has been ignored during version upgrade.", *GetOutermost()->GetName());
 							break;
 						}
 					}
@@ -212,7 +212,7 @@ void UCustomizableObjectNodeModifierRemoveMeshBlocks::BackwardsCompatibleFixup(i
 
 					if (!bFoundInParent)
 					{
-						UE_LOG(LogMutable, Warning, TEXT("[%s] UCustomizableObjectNodeModifierRemoveMeshBlocks refers to and invalid layout block. It has been ignored during version upgrade."), *GetOutermost()->GetName());
+						UE_LOGF(LogMutable, Warning, "[%ls] UCustomizableObjectNodeModifierRemoveMeshBlocks refers to and invalid layout block. It has been ignored during version upgrade.", *GetOutermost()->GetName());
 					}
 				}
 			}
diff --git a/Engine/Plugins/Mutable/Source/CustomizableObjectEditor/Private/MuCOE/Nodes/CustomizableObjectNodeTable.cpp b/Engine/Plugins/Mutable/Source/CustomizableObjectEditor/Private/MuCOE/Nodes/CustomizableObjectNodeTable.cpp
--- a/Engine/Plugins/Mutable/Source/CustomizableObjectEditor/Private/MuCOE/Nodes/CustomizableObjectNodeTable.cpp
+++ b/Engine/Plugins/Mutable/Source/CustomizableObjectEditor/Private/MuCOE/Nodes/CustomizableObjectNodeTable.cpp
@@ -276,7 +276,7 @@ void UCustomizableObjectNodeTable::BackwardsCompatibleFixup(int32 CustomizableOb
 				}
 				else
 				{
-					UE_LOG(LogTemp, Warning, TEXT("BackwardCompatibleFixUp Problem: Property doesn't exist"));
+					UE_LOGF(LogTemp, Warning, "BackwardCompatibleFixUp Problem: Property doesn't exist");
 				}
 			}
 		}
diff --git a/Engine/Plugins/Mutable/Source/CustomizableObjectEditor/Private/MuCOE/SMutableCodeViewer.cpp b/Engine/Plugins/Mutable/Source/CustomizableObjectEditor/Private/MuCOE/SMutableCodeViewer.cpp
--- a/Engine/Plugins/Mutable/Source/CustomizableObjectEditor/Private/MuCOE/SMutableCodeViewer.cpp
+++ b/Engine/Plugins/Mutable/Source/CustomizableObjectEditor/Private/MuCOE/SMutableCodeViewer.cpp
@@ -825,7 +825,7 @@ void SMutableCodeViewer::CacheOperationsMatchingStringPattern()
 	
 	if (!SearchString.IsEmpty())
 	{
-		UE_LOG(LogMutable,Display,TEXT("Starting string search with target string ""\"%s""\" "),*SearchString);
+		UE_LOGF(LogMutable,Display, "Starting string search with target string ""\"%ls""\" ",*SearchString);
 	
 		// Object containing all data required by the search operation to be able to be called recursively
 		FElementsSearchCache SearchPayload;
@@ -840,7 +840,7 @@ void SMutableCodeViewer::CacheOperationsMatchingStringPattern()
 		NameBasedNavigationElements = MoveTemp(SearchPayload.FoundElements);
 		SortElementsByTreeIndex(NameBasedNavigationElements);
 		
-		UE_LOG(LogMutable, Display, TEXT("Operations found with matching pattern ""\"%s""\" is  %i"), *SearchString, NameBasedNavigationElements.Num());
+		UE_LOGF(LogMutable, Display, "Operations found with matching pattern ""\"%ls""\" is  %i", *SearchString, NameBasedNavigationElements.Num());
 	}
 	else
 	{
@@ -2040,7 +2040,7 @@ void SMutableCodeViewer::GetVisibleChildren(TSharedPtr<FMutableCodeTreeElement>
 	}
 
 	// Debug
-	// UE_LOG(LogTemp,Warning,TEXT("Found a total of %i children elements "),OutChildren.Num());
+	// UE_LOGF(LogTemp,Warning, "Found a total of %i children elements ",OutChildren.Num());
 }
 
 
@@ -2509,7 +2509,7 @@ void SMutableCodeViewer::CacheRootNodeAddresses()
 	if (IndexOnConstantsArray < 0)
 	{
 		// Not valid index.
-		UE_LOG(LogTemp,Error,TEXT("The provided index [%d] is not valid."),IndexOnConstantsArray );
+		UE_LOGF(LogTemp,Error, "The provided index [%d] is not valid.",IndexOnConstantsArray );
 		return;
 	}
 
@@ -2537,7 +2537,7 @@ void SMutableCodeViewer::CacheRootNodeAddresses()
 	{
 		TargetedTypeSelector->SetSelectedItem(NoneOperationEntry);
 		NavigationElements.Reset();
-		UE_LOG(LogTemp,Error,TEXT("The provided constant index does not seem to be used anywhere : Make sure the index is valid and that IsConstantResourceUsedByOperation() switch is up to date"));
+		UE_LOGF(LogTemp,Error, "The provided constant index does not seem to be used anywhere : Make sure the index is valid and that IsConstantResourceUsedByOperation() switch is up to date");
 	}
 
 	// Reset the navigation index
@@ -2903,7 +2903,7 @@ void SMutableCodeViewer::Tick(const FGeometry& AllottedGeometry, const double In
 	
 	default:
 #if UE_BUILD_DEBUG || UE_BUILD_DEVELOPMENT
-		UE_LOG(LogMutable, Log, TEXT("There is no previewer for the selected type of Mutable object"))
+		UE_LOGF(LogMutable, Log, "There is no previewer for the selected type of Mutable object")
 #endif
 		PreviewBorder->SetContent(SNullWidget::NullWidget);
 		// There is no viewer for this type.
@@ -3055,7 +3055,7 @@ void SMutableCodeViewer::PreviewMutableProjector(const UE::Mutable::Private::FPr
 {
 	if (!Projector)
 	{
-		UE_LOG(LogTemp,Error,TEXT("Unable to preview data on null Projector pointer."))
+		UE_LOGF(LogTemp,Error, "Unable to preview data on null Projector pointer.")
 		return;
 	}
 	
@@ -3092,19 +3092,19 @@ void SMutableCodeViewer::PreviewMutableCurve(const FRichCurve& Curve)
 // TODO: Implement physics viewer
 void SMutableCodeViewer::PreviewMutablePhysics(UE::Mutable::Private::TManagedPtr<const UE::Mutable::Private::FPhysicsBody> Physics)
 {
-	UE_LOG(LogMutable, Warning, TEXT("Previewer for Mutable Physics not yet implemented"))
+	UE_LOGF(LogMutable, Warning, "Previewer for Mutable Physics not yet implemented")
 }
 
 // TODO: Implement matrix viewer
 void SMutableCodeViewer::PreviewMutableMatrix(const FMatrix44f& Mat)
 {
-	UE_LOG(LogMutable, Warning, TEXT("Previewer for Mutable Matrices not yet implemented"))
+	UE_LOGF(LogMutable, Warning, "Previewer for Mutable Matrices not yet implemented")
 }
 
 // TODO: Implement shape viewer
 void SMutableCodeViewer::PreviewMutableShape(const UE::Mutable::Private::FShape* Shape)
 {
-	UE_LOG(LogMutable, Warning, TEXT("Previewer for Mutable Shapes not yet implemented"))
+	UE_LOGF(LogMutable, Warning, "Previewer for Mutable Shapes not yet implemented")
 }
 
 
diff --git a/Engine/Plugins/Mutable/Source/CustomizableObjectEditor/Private/MuCOE/SMutableConstantsWidget.cpp b/Engine/Plugins/Mutable/Source/CustomizableObjectEditor/Private/MuCOE/SMutableConstantsWidget.cpp
--- a/Engine/Plugins/Mutable/Source/CustomizableObjectEditor/Private/MuCOE/SMutableConstantsWidget.cpp
+++ b/Engine/Plugins/Mutable/Source/CustomizableObjectEditor/Private/MuCOE/SMutableConstantsWidget.cpp
@@ -1461,7 +1461,7 @@ void SMutableConstantsWidget::ContractExpandableAreas(const TSharedPtr<SExpandab
 #if UE_BUILD_DEBUG || UE_BUILD_DEVELOPMENT
 	if(!InException)
 	{
-		UE_LOG(LogTemp,Warning,TEXT("No valid expandable area has been provided as exception : All expandable areas will therefore get contracted"));
+		UE_LOGF(LogTemp,Warning, "No valid expandable area has been provided as exception : All expandable areas will therefore get contracted");
 	}
 #endif
 	
diff --git a/Engine/Plugins/Mutable/Source/CustomizableObjectEditor/Private/MuCOE/SMutableGraphViewer.cpp b/Engine/Plugins/Mutable/Source/CustomizableObjectEditor/Private/MuCOE/SMutableGraphViewer.cpp
--- a/Engine/Plugins/Mutable/Source/CustomizableObjectEditor/Private/MuCOE/SMutableGraphViewer.cpp
+++ b/Engine/Plugins/Mutable/Source/CustomizableObjectEditor/Private/MuCOE/SMutableGraphViewer.cpp
@@ -981,7 +981,7 @@ void SMutableGraphViewer::GetChildrenForInfo(TSharedPtr<FMutableGraphTreeElement
 	
 	else
 	{
-		UE_LOG(LogMutable,Error,TEXT("The node of type %d has not been implemented, so its children won't be added to the tree."), int32(ParentNodeType->Type));
+		UE_LOGF(LogMutable,Error, "The node of type %d has not been implemented, so its children won't be added to the tree.", int32(ParentNodeType->Type));
 
 		// Add a placeholder to the tree
 		const FString Prefix =  FString::Printf(TEXT("[%d] NODE TYPE NOT IMPLEMENTED"), int32(ParentNodeType->Type));
diff --git a/Engine/Plugins/Mutable/Source/CustomizableObjectEditor/Private/MuCOE/SMutableMeshViewer.cpp b/Engine/Plugins/Mutable/Source/CustomizableObjectEditor/Private/MuCOE/SMutableMeshViewer.cpp
--- a/Engine/Plugins/Mutable/Source/CustomizableObjectEditor/Private/MuCOE/SMutableMeshViewer.cpp
+++ b/Engine/Plugins/Mutable/Source/CustomizableObjectEditor/Private/MuCOE/SMutableMeshViewer.cpp
@@ -224,7 +224,7 @@ void SMutableMeshViewer::SetMesh(const UE::Mutable::Private::TManagedPtr<const U
 
 			FString DebugLog;
 			MutableMesh->Log(DebugLog, 8);
-			UE_LOG(LogMutable, Warning, TEXT("[%s]"), *DebugLog);
+			UE_LOGF(LogMutable, Warning, "[%ls]", *DebugLog);
 		}
 
 		MeshViewport->SetMesh(MutableMesh);
diff --git a/Engine/Plugins/Mutable/Source/CustomizableObjectEditor/Private/MuCOE/UnrealBakeHelpers.cpp b/Engine/Plugins/Mutable/Source/CustomizableObjectEditor/Private/MuCOE/UnrealBakeHelpers.cpp
--- a/Engine/Plugins/Mutable/Source/CustomizableObjectEditor/Private/MuCOE/UnrealBakeHelpers.cpp
+++ b/Engine/Plugins/Mutable/Source/CustomizableObjectEditor/Private/MuCOE/UnrealBakeHelpers.cpp
@@ -162,7 +162,7 @@ UTexture2D* FUnrealBakeHelpers::BakeHelper_CreateAssetTexture(UTexture2D* Source
 	int32 MipCount = SourceTexture->GetPlatformData()->Mips.Num();
 	if (!MipCount)
 	{
-		UE_LOG(LogMutable, Warning, TEXT("Bake Instances: Empty texture found [%s]."), *SourceTexture->GetName());
+		UE_LOGF(LogMutable, Warning, "Bake Instances: Empty texture found [%ls].", *SourceTexture->GetName());
 		
 		DupTex->PostEditChange();
 		return DupTex;
