[Mutable] Fix mesh streaming bugs when mutable mesh sections are removed.

- Fix clothing and morphs section indices when streaming meshes have removed sections.
- Removed special case for a single section in ClothVertexBuffer.
- Moved dummy triangle for empty sections to the end of the buffer to not invalidate indices references.

#rnx

#Synced-CL 51757264

diff --git a/Engine/Plugins/Mutable/Source/CustomizableObject/Internal/MuCO/UnrealConversionUtils.h b/Engine/Plugins/Mutable/Source/CustomizableObject/Internal/MuCO/UnrealConversionUtils.h
--- a/Engine/Plugins/Mutable/Source/CustomizableObject/Internal/MuCO/UnrealConversionUtils.h
+++ b/Engine/Plugins/Mutable/Source/CustomizableObject/Internal/MuCO/UnrealConversionUtils.h
@@ -109,13 +109,15 @@ namespace UnrealConversionUtils
 	 *Performs a copy of the data found on the index buffers on the mutable mesh to the buffers of the skeletal mesh
 	 * @param LODResource - LODRenderData to be updated.
 	 * @param InMutableMesh - The mutable mesh whose index buffers you want to work with
+	 * @param bOutMarkRenderStateDirty
+	 * @param RenderToMutableSectionIndexMap - Map form RenderSection Index to MutableSection Index. In case the optional is not set, the identity is used.
 	 * @return True if the operation could be performed successfully, false if not.
 	 */
 	CUSTOMIZABLEOBJECT_API bool CopyMutableIndexBuffers(
 		FSkeletalMeshLODRenderData& LODResource,
 		const UE::Mutable::Private::FMesh* InMutableMesh,
-		const TArray<uint32>& SurfaceIDs,
-		bool& bOutMarkRenderStateDirty);
+		bool& bOutMarkRenderStateDirty,
+		TOptional<TConstArrayView<int32>> RenderToMutableSectionIndexMap = {});
 	
 
 	/**
@@ -159,7 +161,13 @@ namespace UnrealConversionUtils
 		FSkeletalMeshLODRenderData& LODResource
 	);
 
-	void MorphTargetVertexInfoBuffers(FSkeletalMeshLODRenderData& LODResource, const USkeletalMesh& Owner, const UE::Mutable::Private::FMesh& MutableMesh, int32 LODIndex, bool bGenerateCPUMorphTargets);
+	void MorphTargetVertexInfoBuffers(
+			FSkeletalMeshLODRenderData& LODResource, 
+			const USkeletalMesh& Owner, 
+			const UE::Mutable::Private::FMesh& MutableMesh, 
+			int32 LODIndex, 
+			bool bGenerateCPUMorphTargets,
+			TOptional<TConstArrayView<int32>> RenderToMutableSectionIndexMap = {});
 
 
 	CUSTOMIZABLEOBJECT_API void SetMorphData(
@@ -174,7 +182,10 @@ namespace UnrealConversionUtils
 	TSharedPtr<UE::Mutable::Private::FPhysicsBody> MakePhysicsBodyFromAsset(UPhysicsAsset* Asset, const TArray<uint8>& BodySetupRelevancyMap);
 	
 	
-	void ClothVertexBuffers(FSkeletalMeshLODRenderData& LODResource, const UE::Mutable::Private::FMesh& MutableMesh);
+	void ClothVertexBuffers(
+			FSkeletalMeshLODRenderData& LODResource, 
+			const UE::Mutable::Private::FMesh& MutableMesh,
+			TOptional<TConstArrayView<int32>> RenderToMutableSectionIndexMap = {});
 
 	CUSTOMIZABLEOBJECT_API void PropagateBoneUsageFlagsThroughMeshPose(UE::Mutable::Private::FMesh& Mesh);
 
diff --git a/Engine/Plugins/Mutable/Source/CustomizableObject/Private/MuCO/CustomizableObjectInstance.cpp b/Engine/Plugins/Mutable/Source/CustomizableObject/Private/MuCO/CustomizableObjectInstance.cpp
--- a/Engine/Plugins/Mutable/Source/CustomizableObject/Private/MuCO/CustomizableObjectInstance.cpp
+++ b/Engine/Plugins/Mutable/Source/CustomizableObject/Private/MuCO/CustomizableObjectInstance.cpp
@@ -3728,7 +3728,7 @@ void BuildOrCopyClothingData(
 		TArray<FClothBufferIndexMapping> Mappings;
 
 		constexpr int32 Stride = sizeof(FMeshToMeshVertData);
-			
+		
 		int32 NumVertices = 0; // Upper bound
 			
 		const int32 NumSections = LOD->Mesh->Surfaces.Num();
@@ -3766,7 +3766,7 @@ void BuildOrCopyClothingData(
 		{
 			continue;
 		}
-		
+
 		FSkeletalMeshLODRenderData& RenderLOD = SkeletalMesh->GetResourceForRendering()->LODRenderData[LODIndex];
 		UnrealConversionUtils::ClothVertexBuffers(RenderLOD, *LOD->Mesh);
 	}
@@ -4025,7 +4025,7 @@ bool BuildOrCopyRenderData(const TSharedRef<FUpdateContextPrivate>& Context, USk
 
 			// Copy indices.
 			bool bMarkRenderStateDirty = false;
-			if (!UnrealConversionUtils::CopyMutableIndexBuffers(LODResource, LOD->Mesh.Get(), SurfaceIDs, bMarkRenderStateDirty))
+			if (!UnrealConversionUtils::CopyMutableIndexBuffers(LODResource, LOD->Mesh.Get(), bMarkRenderStateDirty))
 			{
 				// End with failure
 				return false;
diff --git a/Engine/Plugins/Mutable/Source/CustomizableObject/Private/MuCO/CustomizableObjectMeshUpdate.cpp b/Engine/Plugins/Mutable/Source/CustomizableObject/Private/MuCO/CustomizableObjectMeshUpdate.cpp
--- a/Engine/Plugins/Mutable/Source/CustomizableObject/Private/MuCO/CustomizableObjectMeshUpdate.cpp
+++ b/Engine/Plugins/Mutable/Source/CustomizableObject/Private/MuCO/CustomizableObjectMeshUpdate.cpp
@@ -299,7 +299,6 @@ void FCustomizableObjectMeshStreamIn::RequestMeshUpdate(const FContext& Context)
 	}
 }
 
-
 void FCustomizableObjectMeshStreamIn::ConvertMesh(const FContext& Context, bool& bOutMarkRenderStateDirty)
 {
 	MUTABLE_CPUPROFILER_SCOPE(FCustomizableObjectMeshStreamIn::ConvertMesh);
@@ -332,18 +331,46 @@ void FCustomizableObjectMeshStreamIn::ConvertMesh(const FContext& Context, bool&
 			Abort();
 			return;
 		}
-		
+
+
 		const bool bNeedsCPUAccess = Mesh->GetResourceForRendering()->RequiresCPUSkinning(GMaxRHIFeatureLevel) || Mesh->NeedCPUData(LODIndexAbsolute);
 		const bool bGenerateCPUMorphTargetsIfNeeded = bEnableCPUMorphTargetStreaming;
-
+		
 		FSkeletalMeshLODRenderData& LODResource = *Context.LODResourcesView[LODIndex];
+
+		TOptional<TConstArrayView<int32>> RenderToMutableSectionIndexMap;
+
+		TArray<int32, TInlineAllocator<16>> RenderToMutableSectionIndexMapStorage;
+		if (LODResource.RenderSections.Num() != MutableMesh->Surfaces.Num())
+		{
+			const int32 NumRenderSections = LODResource.RenderSections.Num();
+			RenderToMutableSectionIndexMapStorage.Init(INDEX_NONE, NumRenderSections);
+			
+			TArray<uint32>& SurfaceIDs = UpdateContext->MeshContext->SurfaceIDs[LODIndexAbsolute];
+
+			check(SurfaceIDs.Num() == NumRenderSections);
+			for (int32 SectionIndex = 0; SectionIndex < NumRenderSections; ++SectionIndex)
+			{
+				const int32 SurfaceID = SurfaceIDs[SectionIndex];
+				
+				RenderToMutableSectionIndexMapStorage[SectionIndex] = MutableMesh->Surfaces.IndexOfByPredicate(
+				[SurfaceID](const UE::Mutable::Private::FMeshSurface& Surface)
+				{
+					return Surface.Id == SurfaceID;
+				});
+			}
+
+			RenderToMutableSectionIndexMap.Emplace(RenderToMutableSectionIndexMapStorage);
+		}
+
 		UnrealConversionUtils::CopyMutableVertexBuffers(LODResource, MutableMesh.Get(), bNeedsCPUAccess);
-		UnrealConversionUtils::CopyMutableIndexBuffers(LODResource, MutableMesh.Get(), UpdateContext->MeshContext->SurfaceIDs[LODIndexAbsolute], bOutMarkRenderStateDirty);
+		UnrealConversionUtils::CopyMutableIndexBuffers(LODResource, MutableMesh.Get(), bOutMarkRenderStateDirty, RenderToMutableSectionIndexMap);
 		UnrealConversionUtils::CopyMutableSkinWeightProfilesBuffers(LODResource, *const_cast<UCustomizableObjectSkeletalMesh*>(Mesh), LODIndexAbsolute, MutableMesh.Get(), UpdateContext->MeshContext->SkinWeightProfileIDs);
-		UnrealConversionUtils::MorphTargetVertexInfoBuffers(LODResource, *Mesh, *MutableMesh.Get(), LODIndexAbsolute, bGenerateCPUMorphTargetsIfNeeded);
-		UnrealConversionUtils::ClothVertexBuffers(LODResource, *MutableMesh.Get());
+		UnrealConversionUtils::MorphTargetVertexInfoBuffers(LODResource, *Mesh, *MutableMesh.Get(), LODIndexAbsolute, bGenerateCPUMorphTargetsIfNeeded, RenderToMutableSectionIndexMap);
+		UnrealConversionUtils::ClothVertexBuffers(LODResource, *MutableMesh.Get(), RenderToMutableSectionIndexMap);
 		
 		UnrealConversionUtils::UpdateSkeletalMeshLODRenderDataBuffersSize(LODResource);
+
 	}
 
 	// Clear MeshUpdate data
diff --git a/Engine/Plugins/Mutable/Source/CustomizableObject/Private/MuCO/UnrealConversionUtils.cpp b/Engine/Plugins/Mutable/Source/CustomizableObject/Private/MuCO/UnrealConversionUtils.cpp
--- a/Engine/Plugins/Mutable/Source/CustomizableObject/Private/MuCO/UnrealConversionUtils.cpp
+++ b/Engine/Plugins/Mutable/Source/CustomizableObject/Private/MuCO/UnrealConversionUtils.cpp
@@ -535,7 +535,6 @@ namespace UnrealConversionUtils
 		}
 	}
 
-		
 	void InitIndexBuffersWithDummyData(FSkeletalMeshLODRenderData& LODResource, const UE::Mutable::Private::FMesh* InMutableMesh)
 	{
 		MUTABLE_CPUPROFILER_SCOPE(InitIndexBuffersWithDummyData);
@@ -549,12 +548,11 @@ namespace UnrealConversionUtils
 		LODResource.MultiSizeIndexContainer.GetIndexBuffer()->SetMetaData(NumIndices);
 	}
 
-
 	bool CopyMutableIndexBuffers(
 		FSkeletalMeshLODRenderData& LODResource,
 		const UE::Mutable::Private::FMesh* InMutableMesh,
-		const TArray<uint32>& SurfaceIDs,
-		bool& bOutMarkRenderStateDirty)
+		bool& bOutMarkRenderStateDirty,
+		TOptional<TConstArrayView<int32>> RenderToMutableSectionIndexMap)
 	{
 		MUTABLE_CPUPROFILER_SCOPE(CopyMutableIndexBuffers);
 
@@ -585,27 +583,33 @@ namespace UnrealConversionUtils
 		
 		// UE-363228: Add a collapsed triangle whenever instead of disabling missing sections
 		const bool bAddCollapsedTriangle = LODResource.RenderSections.Num() != InMutableMesh->Surfaces.Num();
-		const int32 FirstIndex = bAddCollapsedTriangle ? 3 : 0;
 
 		// Don't assume the buffer is empty, otherwise we may add extra elements using Insert(). 
 		{
-			const int32 IndexCount = MutableIndexCount + FirstIndex;
+			const int32 IndexCount = MutableIndexCount + 3*(int32)bAddCollapsedTriangle;
 			LODResource.MultiSizeIndexContainer.GetIndexBuffer()->Empty(IndexCount);
 			LODResource.MultiSizeIndexContainer.GetIndexBuffer()->Insert(0, IndexCount);
 		}
 
-		FMemory::Memcpy(LODResource.MultiSizeIndexContainer.GetIndexBuffer()->GetPointerTo(FirstIndex), DataPtr, MutableIndexCount * ElementSize);
-		
+		FMemory::Memcpy(LODResource.MultiSizeIndexContainer.GetIndexBuffer()->GetPointerTo(0), DataPtr, MutableIndexCount * ElementSize);
+	
+		check(!RenderToMutableSectionIndexMap || RenderToMutableSectionIndexMap->Num() == LODResource.RenderSections.Num());
 		for (int32 SectionIndex = 0; SectionIndex < LODResource.RenderSections.Num(); ++SectionIndex)
 		{
 			FSkelMeshRenderSection& Section = LODResource.RenderSections[SectionIndex];
 			
-			const int32 SurfaceID = SurfaceIDs[SectionIndex];
-			const UE::Mutable::Private::FMeshSurface* Surface = InMutableMesh->Surfaces.FindByPredicate([SurfaceID](const UE::Mutable::Private::FMeshSurface& Surface)
+			int32 MutableSectionIndex = SectionIndex;
+			if (RenderToMutableSectionIndexMap)
 			{
-				return Surface.Id == SurfaceID;
-			});
+				MutableSectionIndex = (*RenderToMutableSectionIndexMap)[SectionIndex];
+			}
 			
+			const UE::Mutable::Private::FMeshSurface* Surface = nullptr;
+			if (MutableSectionIndex != INDEX_NONE)
+			{
+				Surface = &InMutableMesh->Surfaces[MutableSectionIndex];
+			}
+
 			if (Surface)
 			{
 				const int32 NumVertices = Surface->SubMeshes.Last().VertexEnd - Surface->SubMeshes[0].VertexBegin;
@@ -615,14 +619,18 @@ namespace UnrealConversionUtils
 
 				const int32 IndexBegin = Surface->SubMeshes[0].IndexBegin; 
 				const int32 IndexEnd = Surface->SubMeshes.Last().IndexEnd; 
-				
+			
+				check(IndexBegin <= IndexEnd);
+				check(IndexEnd*ElementSize <= LODResource.MultiSizeIndexContainer.GetIndexBuffer()->GetResourceDataSize())
+
+				check((IndexEnd - IndexBegin) % 3 == 0);
 				const int32 NumTriangles = (IndexEnd - IndexBegin) / 3;
 
 				bOutMarkRenderStateDirty |= Section.NumTriangles != NumTriangles;
 				Section.NumTriangles = NumTriangles;
 
 				bOutMarkRenderStateDirty |= Section.BaseIndex != IndexBegin;
-				Section.BaseIndex = IndexBegin + FirstIndex;
+				Section.BaseIndex = IndexBegin;
 
 				const int32 VertexBegin = Surface->SubMeshes[0].VertexBegin;
 
@@ -634,7 +642,7 @@ namespace UnrealConversionUtils
 				// Section.bDisabled = true; // UE-363228: Disabling sections after the mesh is generated produces visual artifacts.
 				Section.NumTriangles = 1; // UE-363228: Add collapsed triangle to avoid visual artifacts
 				Section.NumVertices = 0;
-				Section.BaseIndex = 0;
+				Section.BaseIndex = MutableIndexCount;
 				Section.BaseVertexIndex = 0;
 
 				bOutMarkRenderStateDirty = true;	
@@ -645,6 +653,7 @@ namespace UnrealConversionUtils
 	}
 
 
+
 	struct FAuxMutableSkinWeightVertexKey
 	{
 		FAuxMutableSkinWeightVertexKey(const uint8* InKey, uint8 InKeySize, uint32 InHash)
@@ -928,7 +937,14 @@ namespace UnrealConversionUtils
 		}
 	}
 
-	void MorphTargetVertexInfoBuffers(FSkeletalMeshLODRenderData& LODResource, const USkeletalMesh& Owner, const UE::Mutable::Private::FMesh& MutableMesh, int32 LODIndex, bool bGenerateCPUMorphTargetsIfNeeded)
+	void MorphTargetVertexInfoBuffers(
+			FSkeletalMeshLODRenderData& LODResource, 
+			const USkeletalMesh& Owner, 
+			const UE::Mutable::Private::FMesh& MutableMesh, 
+			int32 LODIndex, 
+			bool bGenerateCPUMorphTargetsIfNeeded,
+			TOptional<TConstArrayView<int32>> RenderToMutableSectionIndexMap
+	)
 	{
 		if (Owner.GetMorphTargets().IsEmpty())
 		{
@@ -940,6 +956,22 @@ namespace UnrealConversionUtils
 			return;
 		}
 
+		TOptional<TArray<int32, TInlineAllocator<16>>> ReverseSectionIndexMap;
+		if (RenderToMutableSectionIndexMap)
+		{
+			ReverseSectionIndexMap.Emplace();
+			ReverseSectionIndexMap->Init(INDEX_NONE, MutableMesh.Surfaces.Num());
+
+			for (int32 I = 0; I < RenderToMutableSectionIndexMap->Num(); ++I)
+			{
+				int32 SectionIndex = (*RenderToMutableSectionIndexMap)[I];
+				if (SectionIndex != INDEX_NONE)
+				{
+					(*ReverseSectionIndexMap)[SectionIndex] = I;
+				}
+			}
+		}
+
 		// Get the global names map.
 		const TMap<FName, int32>& IndexMap = Owner.GetMorphTargetIndexMap();
 	
@@ -981,7 +1013,7 @@ namespace UnrealConversionUtils
 
 		MorphTargetMinimumValues.SetNum(Owner.GetMorphTargets().Num());
 		MorphTargetMaximumValues.SetNum(Owner.GetMorphTargets().Num());
-	
+
 		TArray<FMorphTargetCompressedLODModel*> CompressedMorphsToInit;
 		CompressedMorphsToInit.Init(nullptr, Owner.GetMorphTargets().Num());
 
@@ -999,7 +1031,15 @@ namespace UnrealConversionUtils
 				UMorphTarget* MorphLOD = Owner.GetMorphTargets()[*FoundIndex];
 
 				check(MorphLOD->GetMorphLODModels().IsValidIndex(LODIndex));
+
 				MorphLOD->GetMorphLODModels()[LODIndex].SectionIndices = MutableMesh.Morph.SurfacesInUsePerMorph[MorphIndex];
+				if (ReverseSectionIndexMap)
+				{
+					for (int32& SectionIndex : MorphLOD->GetMorphLODModels()[LODIndex].SectionIndices)
+					{
+						SectionIndex = (*ReverseSectionIndexMap)[SectionIndex];
+					}
+				}
 			}
 		}
 
@@ -1989,88 +2029,74 @@ namespace UnrealConversionUtils
 	}
 	
 
-	void ClothVertexBuffers(FSkeletalMeshLODRenderData& LODResource, const UE::Mutable::Private::FMesh& MutableMesh)
+	void ClothVertexBuffers(
+			FSkeletalMeshLODRenderData& LODResource, 
+			const UE::Mutable::Private::FMesh& MutableMesh,
+			TOptional<TConstArrayView<int32>> RenderToMutableSectionIndexMap)
 	{
-		const int32 NumSections = MutableMesh.Surfaces.Num();
-		check(NumSections == MutableMesh.ClothSections.Num());
+		const int32 NumSections = MutableMesh.ClothSections.Num();
+		check(MutableMesh.Surfaces.Num() == MutableMesh.ClothSections.Num());
 		
-		if (NumSections == 1) // Single section. No need to reconstruct data.
+		int32 DataSize = 0;
+		for (int32 SectionIndex = 0; SectionIndex < NumSections; ++SectionIndex)
 		{
-			const uint32 NumClothVertices = MutableMesh.ClothSections[0].Data.Num();
-			if (MutableMesh.ClothSections[0].IsValid() && NumClothVertices)
-			{
-				FClothBufferIndexMapping Mapping;
-				Mapping.BaseVertexIndex = 0;
-				Mapping.MappingOffset = 0;
-				Mapping.LODBiasStride = NumClothVertices;
-		
-				TArray<FClothBufferIndexMapping> Mappings;
-				Mappings.Add(Mapping);
-			
-				TArrayView<const FMeshToMeshVertData> DataView = TArrayView(MutableMesh.ClothSections[0].Data.GetData(), NumClothVertices);
-				LODResource.ClothVertexBuffer.Init(DataView, Mappings);
-			
-				// Cloth hack. This data should not be used.
-				FSkelMeshRenderSection& RenderSection = LODResource.RenderSections[0];
-				RenderSection.ClothMappingDataLODs.SetNum(1, EAllowShrinking::No);
-				RenderSection.ClothMappingDataLODs[0].SetNum(NumClothVertices, EAllowShrinking::No);
-				FMemory::Memcpy(&RenderSection.ClothMappingDataLODs[0][0], MutableMesh.ClothSections[0].Data.GetData(), NumClothVertices * sizeof(FMeshToMeshVertData));
-			}
+			DataSize += MutableMesh.ClothSections[SectionIndex].Data.Num();
 		}
-		else if (NumSections > 1)
+	
+		if (DataSize)
 		{
-			int32 DataSize = 0;
-			for (int32 SectionIndex = 0; SectionIndex < NumSections; ++SectionIndex)
+			TArray<FMeshToMeshVertData> Data;
+			Data.Reserve(DataSize);
+
+			for (const UE::Mutable::Private::FCloth& Cloth : MutableMesh.ClothSections)
 			{
-				DataSize += MutableMesh.ClothSections[SectionIndex].Data.Num();
+				Data.Append(Cloth.Data);
 			}
-		
-			if (DataSize)
+	
+			int32 NumRenderSections = LODResource.RenderSections.Num();
+
+			TArray<FClothBufferIndexMapping> Mappings;
+			Mappings.SetNumZeroed(NumRenderSections);
+
+			check(!RenderToMutableSectionIndexMap || RenderToMutableSectionIndexMap->Num() == NumRenderSections);
+			uint32 VertexOffset = 0;
+			for (int32 SectionIndex = 0; SectionIndex < NumRenderSections; ++SectionIndex)
 			{
-				TArray<FMeshToMeshVertData> Data;
-				Data.Reserve(DataSize);
+				int32 MutableSectionIndex = SectionIndex;
+				if (RenderToMutableSectionIndexMap)
+				{
+					MutableSectionIndex = (*RenderToMutableSectionIndexMap)[SectionIndex];
+				}
 
-				for (const UE::Mutable::Private::FCloth& Cloth : MutableMesh.ClothSections)
+				if (MutableSectionIndex == INDEX_NONE)
 				{
-					Data.Append(Cloth.Data);
+					continue;
 				}
-			
-				TArray<FClothBufferIndexMapping> Mappings;
-				Mappings.SetNum(NumSections);
 
-				uint32 VertexOffset = 0;
-				for (int32 SectionIndex = 0; SectionIndex < NumSections; ++SectionIndex)
+				const UE::Mutable::Private::FCloth& Cloth = MutableMesh.ClothSections[MutableSectionIndex];
+				const uint32 BaseVertexIndex = LODResource.RenderSections[SectionIndex].BaseVertexIndex;
+			
+				FClothBufferIndexMapping& Mapping = Mappings[SectionIndex];
+				
+				const int32 NumClothVertices = Cloth.Data.Num();
+				if (Cloth.IsValid() && NumClothVertices)
 				{
-					const UE::Mutable::Private::FCloth& Cloth = MutableMesh.ClothSections[SectionIndex];
-					const uint32 BaseVertexIndex = LODResource.RenderSections[SectionIndex].BaseVertexIndex;
+					Mapping.BaseVertexIndex = BaseVertexIndex; 
+					Mapping.MappingOffset = VertexOffset;
+					Mapping.LODBiasStride = VertexOffset;
 				
-					FClothBufferIndexMapping& Mapping = Mappings[SectionIndex];
-					
-					const int32 NumClothVertices = Cloth.Data.Num();
-					if (Cloth.IsValid() && NumClothVertices)
-					{
-						Mapping.BaseVertexIndex = BaseVertexIndex; 
-						Mapping.MappingOffset = VertexOffset;
-						Mapping.LODBiasStride = VertexOffset;
-					
-						VertexOffset += NumClothVertices;
-	
-						// Cloth hack. This data should not be used.
-						FSkelMeshRenderSection& RenderSection = LODResource.RenderSections[SectionIndex];
-						RenderSection.ClothMappingDataLODs.SetNum(1, EAllowShrinking::No);
-						RenderSection.ClothMappingDataLODs[0].SetNum(NumClothVertices, EAllowShrinking::No);
-						FMemory::Memcpy(&RenderSection.ClothMappingDataLODs[0][0], Cloth.Data.GetData(), NumClothVertices * sizeof(FMeshToMeshVertData));
-					}
-					else
-					{
-						Mapping.BaseVertexIndex = 0; 
-						Mapping.MappingOffset = 0;
-						Mapping.LODBiasStride = 0;
-					}					
-				}
+					VertexOffset += NumClothVertices;
 
-				LODResource.ClothVertexBuffer.Init(Data, Mappings);
+					// Cloth hack. This data should not be used.
+					FSkelMeshRenderSection& RenderSection = LODResource.RenderSections[SectionIndex];
+					RenderSection.ClothMappingDataLODs.SetNum(1, EAllowShrinking::No);
+					RenderSection.ClothMappingDataLODs[0].SetNum(NumClothVertices, EAllowShrinking::No);
+					FMemory::Memcpy(&RenderSection.ClothMappingDataLODs[0][0], Cloth.Data.GetData(), NumClothVertices * sizeof(FMeshToMeshVertData));
+				}
+				// else we need to add zeros to the Mapping, already done by the initialization.
 			}
+
+			LODResource.ClothVertexBuffer.Init(Data, Mappings);
 		}
 	}
 	
diff --git a/Engine/Plugins/Mutable/Source/CustomizableObjectEditor/Private/MuCOE/MutableMeshPreviewUtils.cpp b/Engine/Plugins/Mutable/Source/CustomizableObjectEditor/Private/MuCOE/MutableMeshPreviewUtils.cpp
--- a/Engine/Plugins/Mutable/Source/CustomizableObjectEditor/Private/MuCOE/MutableMeshPreviewUtils.cpp
+++ b/Engine/Plugins/Mutable/Source/CustomizableObjectEditor/Private/MuCOE/MutableMeshPreviewUtils.cpp
@@ -437,7 +437,7 @@ namespace MutableMeshPreviewUtils
 
 			// Copy index buffers or fail to generate the mesh
 			bool bMarkRenderStateDirty = false;
-			if (!UnrealConversionUtils::CopyMutableIndexBuffers(LODResource, InMutableMesh, SurfaceIDs, bMarkRenderStateDirty))
+			if (!UnrealConversionUtils::CopyMutableIndexBuffers(LODResource, InMutableMesh, bMarkRenderStateDirty))
 			{
 				return false;
 			}
