[Mutable] Change all instances of UE_LOG to UE_LOGF in MutableValidation module.

#rnx
#Synced-CL 51716106

diff --git a/Engine/Plugins/Mutable/Source/MutableValidation/Private/MuV/AssetValidator_CustomizableObjects.cpp b/Engine/Plugins/Mutable/Source/MutableValidation/Private/MuV/AssetValidator_CustomizableObjects.cpp
--- a/Engine/Plugins/Mutable/Source/MutableValidation/Private/MuV/AssetValidator_CustomizableObjects.cpp
+++ b/Engine/Plugins/Mutable/Source/MutableValidation/Private/MuV/AssetValidator_CustomizableObjects.cpp
@@ -93,7 +93,7 @@ EDataValidationResult UAssetValidator_CustomizableObjects::IsCustomizableObjectV
 {
 	EDataValidationResult Result = EDataValidationResult::NotValidated;
 
-	UE_LOG(LogMutable,Verbose,TEXT("Running data validation checks for %s CO."),*InCustomizableObject->GetName());
+	UE_LOGF(LogMutable,Verbose, "Running data validation checks for %ls CO.",*InCustomizableObject->GetName());
 
 	// Bind the post validation method to the post validation delegate if not bound already to be able to know when the validation
 	// operation (for all assets) concludes
diff --git a/Engine/Plugins/Mutable/Source/MutableValidation/Private/MuV/COIBakingTestCommandlet.cpp b/Engine/Plugins/Mutable/Source/MutableValidation/Private/MuV/COIBakingTestCommandlet.cpp
--- a/Engine/Plugins/Mutable/Source/MutableValidation/Private/MuV/COIBakingTestCommandlet.cpp
+++ b/Engine/Plugins/Mutable/Source/MutableValidation/Private/MuV/COIBakingTestCommandlet.cpp
@@ -53,7 +53,7 @@ void OnCustomizableObjectCompilationEnd(const FCompileCallbackParams& Compilatio
 void OnInstanceUpdateEnd(const FUpdateContext& Result)
 {
 	const EUpdateResult InstanceUpdateResult = Result.UpdateResult;
-	UE_LOG(LogMutable,Display,TEXT("Instance finished update with state : %s."), *UEnum::GetValueAsString(InstanceUpdateResult));
+	UE_LOGF(LogMutable,Display, "Instance finished update with state : %ls.", *UEnum::GetValueAsString(InstanceUpdateResult));
 	bWasInstanceUpdateSuccessful = UCustomizableObjectSystem::IsUpdateResultValid(InstanceUpdateResult);
 	
 	// Clear update flag so we can exit the update while loop
@@ -76,7 +76,7 @@ int32 UCOIBakingTestCommandlet::Main(const FString& Params)
 		FString CustomizableObjectInstanceAssetPath = "";
 		if (!FParse::Value(*Params, TEXT("CustomizableObjectInstance="), CustomizableObjectInstanceAssetPath))
 		{
-			UE_LOG(LogMutable,Error,TEXT("Failed to parse Customizable Object Instance package name from provided argument : %s"),*Params)
+			UE_LOGF(LogMutable,Error, "Failed to parse Customizable Object Instance package name from provided argument : %ls",*Params)
 			return 1;
 		}
 	
@@ -84,7 +84,7 @@ int32 UCOIBakingTestCommandlet::Main(const FString& Params)
 		UObject* FoundObject = UE::Mutable::Private::LoadObject(FSoftObjectPath(CustomizableObjectInstanceAssetPath));
 		if (!FoundObject)
 		{
-			UE_LOG(LogMutable,Error,TEXT("Failed to retrieve UObject from path %s"),*CustomizableObjectInstanceAssetPath);
+			UE_LOGF(LogMutable,Error, "Failed to retrieve UObject from path %ls",*CustomizableObjectInstanceAssetPath);
 			return 1;
 		}
 	
@@ -92,11 +92,11 @@ int32 UCOIBakingTestCommandlet::Main(const FString& Params)
 		TargetInstance = Cast<UCustomizableObjectInstance>(FoundObject);
 		if (!TargetInstance)
 		{
-			UE_LOG(LogMutable,Error,TEXT("Failed to cast found UObject to UCustomizableObjectInstance."));
+			UE_LOGF(LogMutable,Error, "Failed to cast found UObject to UCustomizableObjectInstance.");
 			return 1;
 		}
 
-		UE_LOG(LogMutable,Display,TEXT("Successfully loaded %s Customizable Object Instance!"), *TargetInstance->GetName());
+		UE_LOGF(LogMutable,Display, "Successfully loaded %ls Customizable Object Instance!", *TargetInstance->GetName());
 	}
 
 	// Perform a blocking search to ensure all assets used by mutable are reachable using the AssetRegistry
@@ -118,7 +118,7 @@ int32 UCOIBakingTestCommandlet::Main(const FString& Params)
 	UCustomizableObject* InstanceCustomizableObject = TargetInstance->GetCustomizableObject();
 	if (!InstanceCustomizableObject)
 	{
-		UE_LOG(LogMutable,Error,TEXT("The instance %s does not have a CO to compile : Exiting commandlet."), *TargetInstance->GetName());
+		UE_LOGF(LogMutable,Error, "The instance %ls does not have a CO to compile : Exiting commandlet.", *TargetInstance->GetName());
 		return 1;
 	}
 	
@@ -134,7 +134,7 @@ int32 UCOIBakingTestCommandlet::Main(const FString& Params)
 		// To prevent the warnings related to partial compilation from appearing in CIS
 		constexpr bool bPerformPartialCompilation = false;
 
-		UE_LOG(LogMutable, Display, TEXT("Scheduling Customizable Object compilation."));
+		UE_LOGF(LogMutable, Display, "Scheduling Customizable Object compilation.");
 		ScheduleCOCompilationForBaking(*TargetInstance, CompilationEndDelegate, bPerformPartialCompilation);
 		
 		// Exit the wait process once the compilation finishes
@@ -152,7 +152,7 @@ int32 UCOIBakingTestCommandlet::Main(const FString& Params)
 		
 		if (!bWasCustomizableObjectCompilationSuccessful)
 		{
-			UE_LOG(LogMutable, Error, TEXT("Failed to successfully compile the target COIs CO. Exiting commandlet."));
+			UE_LOGF(LogMutable, Error, "Failed to successfully compile the target COIs CO. Exiting commandlet.");
 			return 1;
 		}
 	}
@@ -170,7 +170,7 @@ int32 UCOIBakingTestCommandlet::Main(const FString& Params)
 		InstanceUpdateDelegate.AddStatic(&OnInstanceUpdateEnd);
 		
 		bIsInstanceBeingUpdated = true;
-		UE_LOG(LogMutable, Display, TEXT("Scheduling instance update."));
+		UE_LOGF(LogMutable, Display, "Scheduling instance update.");
 		ScheduleInstanceUpdateForBaking(*TargetInstance, InstanceUpdateDelegate);
 		
 		// Now tick the engine so the instance gets updated while running in the commandlet context
@@ -189,7 +189,7 @@ int32 UCOIBakingTestCommandlet::Main(const FString& Params)
 		// Check the end status of the instance update
 		if (!bWasInstanceUpdateSuccessful)
 		{
-			UE_LOG(LogMutable, Error, TEXT("Failed to successfully update the target COI. Exiting commandlet."));
+			UE_LOGF(LogMutable, Error, "Failed to successfully update the target COI. Exiting commandlet.");
 			return 1;
 		}
 	}
@@ -212,10 +212,10 @@ int32 UCOIBakingTestCommandlet::Main(const FString& Params)
 		// Delete it then but notify the user since this may mean that we have duplicated COIs
 		if (FPaths::DirectoryExists(GlobalBakingDirectory))
 		{
-			UE_LOG(LogMutable,Warning,TEXT("The directory with path  \" %s \" does already exist. This may be produced by an incomplete execution of a previous test. Clearing it out before continuing..."), *GlobalBakingDirectory);
+			UE_LOGF(LogMutable,Warning, "The directory with path  \" %ls \" does already exist. This may be produced by an incomplete execution of a previous test. Clearing it out before continuing...", *GlobalBakingDirectory);
 			if (!IFileManager::Get().DeleteDirectory(*GlobalBakingDirectory,false,true))
 			{
-				UE_LOG(LogMutable,Error,TEXT("Failed to delete the baking directory at path \" %s \"."), *GlobalBakingDirectory);
+				UE_LOGF(LogMutable,Error, "Failed to delete the baking directory at path \" %ls \".", *GlobalBakingDirectory);
 				return 1;
 			}
 		}
@@ -226,7 +226,7 @@ int32 UCOIBakingTestCommandlet::Main(const FString& Params)
 		// Create a new directory where to save the bake itself
 		if (IFileManager::Get().MakeDirectory(*GlobalBakingDirectory, true))
 		{
-			UE_LOG(LogMutable,Display,TEXT("Starting Instance Baking operation."));
+			UE_LOGF(LogMutable,Display, "Starting Instance Baking operation.");
 			{
 				FBakingConfiguration Configuration;
 				Configuration.OutputFilesBaseName = FString::Printf(TEXT("%s_Bake"), *TargetInstance->GetName());
@@ -252,7 +252,7 @@ int32 UCOIBakingTestCommandlet::Main(const FString& Params)
 					if (ensureMsgf(SavedPackage, TEXT("One entry of the \"SavedPackages\" map is a null pointer. All entries are expected to point to non null objects.")))
 					{
 						const FString SaveTypeString = Enum->GetNameByIndex(static_cast<uint8>(SavedPackagePair.Value.SaveResolutionType)).ToString();
-						UE_LOG(LogMutable, Display, TEXT("Stored package \"%s\" with \"%s\" as save resolution type."), *SavedPackage->GetName(), *SaveTypeString)
+						UE_LOGF(LogMutable, Display, "Stored package \"%ls\" with \"%ls\" as save resolution type.", *SavedPackage->GetName(), *SaveTypeString)
 					}
 				}
 			}
@@ -260,25 +260,25 @@ int32 UCOIBakingTestCommandlet::Main(const FString& Params)
 			// Delete the target directory where we did save the baked instance
 			if (!IFileManager::Get().DeleteDirectory(*GlobalBakingDirectory,true,true))
 			{
-				UE_LOG(LogMutable,Error,TEXT("Failed to delete the baking directory at path \" %s \"."), *GlobalBakingDirectory);
+				UE_LOGF(LogMutable,Error, "Failed to delete the baking directory at path \" %ls \".", *GlobalBakingDirectory);
 				return 1;
 			}
 		}
 		else
 		{
-			UE_LOG(LogMutable,Error,TEXT("Failed to create baking directory at path \" %s \"."), *GlobalBakingDirectory);
+			UE_LOGF(LogMutable,Error, "Failed to create baking directory at path \" %ls \".", *GlobalBakingDirectory);
 			return 1;
 		}
 	}
 	
 	if (bWasBakingSuccessful)
 	{
-		UE_LOG(LogMutable, Display, TEXT("Instance Baking operation has been completed successfully."));
+		UE_LOGF(LogMutable, Display, "Instance Baking operation has been completed successfully.");
 		return 0;
 	}
 	else
 	{
-		UE_LOG(LogMutable, Display, TEXT("Instance Baking operation has been completed with errors."));
+		UE_LOGF(LogMutable, Display, "Instance Baking operation has been completed with errors.");
 		return 1;
 	}
 }
diff --git a/Engine/Plugins/Mutable/Source/MutableValidation/Private/MuV/COIBulkUpdateTestCommandlet.cpp b/Engine/Plugins/Mutable/Source/MutableValidation/Private/MuV/COIBulkUpdateTestCommandlet.cpp
--- a/Engine/Plugins/Mutable/Source/MutableValidation/Private/MuV/COIBulkUpdateTestCommandlet.cpp
+++ b/Engine/Plugins/Mutable/Source/MutableValidation/Private/MuV/COIBulkUpdateTestCommandlet.cpp
@@ -34,18 +34,18 @@ int32 UCOIBulkUpdateTestCommandlet::Main(const FString& Params)
 	{
 		if (!FParse::Value(*Params, TEXT("InstancesPackagePath="), InstancesPackagePath))
 		{
-			UE_LOG(LogMutable,Error,TEXT("Failed to parse path where to find the Customizable Object Instances to update : %s"),*Params)
+			UE_LOGF(LogMutable,Error, "Failed to parse path where to find the Customizable Object Instances to update : %ls",*Params)
 			return 1;
 		}
 
 		if (InstancesPackagePath.IsNone())
 		{
-			UE_LOG(LogMutable,Error,TEXT("The path to scan can not be empty"))
+			UE_LOGF(LogMutable,Error, "The path to scan can not be empty")
 			return 1;
 		}
 	}
 
-	UE_LOG(LogMutable, Log, TEXT("(string) instance_search_path : %s "), *InstancesPackagePath.ToString());
+	UE_LOGF(LogMutable, Log, "(string) instance_search_path : %ls ", *InstancesPackagePath.ToString());
 	
 	// Load the asset registry system so we can proceed without issues
 	PrepareAssetRegistry();
@@ -58,7 +58,7 @@ int32 UCOIBulkUpdateTestCommandlet::Main(const FString& Params)
 	// Early exit if no instances could be found for testing.
 	if (FoundAssetData.IsEmpty())
 	{
-		UE_LOG(LogMutable,Error,TEXT("Aborting Bulk Instance Update Test: No assets could be found at  the provided package path"));
+		UE_LOGF(LogMutable,Error, "Aborting Bulk Instance Update Test: No assets could be found at  the provided package path");
 		return 1;
 	}
 
@@ -75,7 +75,7 @@ int32 UCOIBulkUpdateTestCommandlet::Main(const FString& Params)
 		UCustomizableObject* InstanceCO = LoadedInstance->GetCustomizableObject();
 		if (!InstanceCO)
 		{
-			UE_LOG(LogMutable,Error,TEXT("The instance %s does not have a CO. This instance will not get tested."), *LoadedInstance->GetName());
+			UE_LOGF(LogMutable,Error, "The instance %ls does not have a CO. This instance will not get tested.", *LoadedInstance->GetName());
 			continue;
 		}
 		
@@ -89,12 +89,12 @@ int32 UCOIBulkUpdateTestCommandlet::Main(const FString& Params)
 	// Report the objects found (COs and COIs) -------------------------------
 	
 	const uint32 TotalAmountOfCustomizableObjects = MutableResources.Num();
-	UE_LOG(LogMutable,Display,TEXT("Customizable Objects to compile : %u"), TotalAmountOfCustomizableObjects);
-	UE_LOG(LogMutable, Log, TEXT("(int) customizable_objects_count : %u "), TotalAmountOfCustomizableObjects);
+	UE_LOGF(LogMutable,Display, "Customizable Objects to compile : %u", TotalAmountOfCustomizableObjects);
+	UE_LOGF(LogMutable, Log, "(int) customizable_objects_count : %u ", TotalAmountOfCustomizableObjects);
 
 	const uint32 TotalAmountOfInstances = FoundAssetData.Num();
-	UE_LOG(LogMutable,Display,TEXT("Customizable Object Instances to update : %u"), TotalAmountOfInstances);
-	UE_LOG(LogMutable, Log, TEXT("(int) customizable_object_instances_count : %u "), TotalAmountOfInstances);
+	UE_LOGF(LogMutable,Display, "Customizable Object Instances to update : %u", TotalAmountOfInstances);
+	UE_LOGF(LogMutable, Log, "(int) customizable_object_instances_count : %u ", TotalAmountOfInstances);
 	
 	// Report the amount of instances for each of the COs found as parents of the instances in the target path
 	for (TTuple<UCustomizableObject*, TArray<TStrongObjectPtr<UCustomizableObjectInstance>>>& MutableResourceTuple : MutableResources)
@@ -104,8 +104,8 @@ int32 UCOIBulkUpdateTestCommandlet::Main(const FString& Params)
 		const FString CustomizableObjectName = CustomizableObjectToCompile->GetName();
 		
 		const uint32 COInstancesCount = MutableResourceTuple.Value.Num();
-		UE_LOG(LogMutable,Display,TEXT("The CO \"%s\" has in total \"%u\" instances."), *CustomizableObjectName, COInstancesCount );
-		UE_LOG(LogMutable, Log, TEXT("(int) %s_instance_count : %u "), *CustomizableObjectName, COInstancesCount);
+		UE_LOGF(LogMutable,Display, "The CO \"%ls\" has in total \"%u\" instances.", *CustomizableObjectName, COInstancesCount );
+		UE_LOGF(LogMutable, Log, "(int) %ls_instance_count : %u ", *CustomizableObjectName, COInstancesCount);
 		
 		// print the name of the instances+
 		uint32 InstanceIndex = 0;
@@ -113,7 +113,7 @@ int32 UCOIBulkUpdateTestCommandlet::Main(const FString& Params)
 		{
 			check(Instance);
 			const FString CustomizableObjectInstanceName = Instance->GetName();
-			UE_LOG(LogMutable, Log, TEXT("%u : \"%s\""), InstanceIndex++, *CustomizableObjectInstanceName);
+			UE_LOGF(LogMutable, Log, "%u : \"%ls\"", InstanceIndex++, *CustomizableObjectInstanceName);
 		}
 	}
 	
@@ -153,13 +153,13 @@ int32 UCOIBulkUpdateTestCommandlet::Main(const FString& Params)
 		// Compile the current CO object
 		if (!CompilationUtility->CompileCustomizableObject(*CustomizableObject, false, &CompilationOptions))	// Do not log mutable data since mongoDB will not be able to handle it correctly 
 		{
-			UE_LOG(LogMutable, Error,TEXT("The CO %s could not be compiled successfully. Skipping the update of all COIs that use it."), *CustomizableObjectName )
+			UE_LOGF(LogMutable, Error, "The CO %ls could not be compiled successfully. Skipping the update of all COIs that use it.", *CustomizableObjectName )
 			continue;
 		}
 
 		// Now that CO has been compiled proceed with the update of the instances that use it
 		const uint32 CustomizableObjectInstanceCount = ResourcesIterator->Value.Num();
-		UE_LOG(LogMutable, Display, TEXT("Starting update of the \"%u\" instances with CO : \"%s\"."), CustomizableObjectInstanceCount, *CustomizableObjectName);
+		UE_LOGF(LogMutable, Display, "Starting update of the \"%u\" instances with CO : \"%ls\".", CustomizableObjectInstanceCount, *CustomizableObjectName);
 
 		// Iterate over all the COIs of the CO and update them
 		TArray<TStrongObjectPtr<UCustomizableObjectInstance>>& Instances = ResourcesIterator->Value;
@@ -169,7 +169,7 @@ int32 UCOIBulkUpdateTestCommandlet::Main(const FString& Params)
 			check(Instance);
 			
 			Instances.RemoveAt(0);
-			UE_LOG(LogMutable, Display, TEXT("\t( %u / %u ) Processing instance : \"%s\" ."),CurrentInstanceIndex++, TotalAmountOfInstances ,*Instance->GetName());
+			UE_LOGF(LogMutable, Display, "\t( %u / %u ) Processing instance : \"%ls\" .",CurrentInstanceIndex++, TotalAmountOfInstances ,*Instance->GetName());
 
 			CollectGarbage(RF_NoFlags, true);
 			
@@ -184,6 +184,6 @@ int32 UCOIBulkUpdateTestCommandlet::Main(const FString& Params)
 		CollectGarbage(RF_NoFlags, true);
 	}
 
-	UE_LOG(LogMutable,Display,TEXT("Mutable commandlet finished."));
+	UE_LOGF(LogMutable,Display, "Mutable commandlet finished.");
 	return 0;
 }
diff --git a/Engine/Plugins/Mutable/Source/MutableValidation/Private/MuV/COMultiPluginBulkValidationCommandlet.cpp b/Engine/Plugins/Mutable/Source/MutableValidation/Private/MuV/COMultiPluginBulkValidationCommandlet.cpp
--- a/Engine/Plugins/Mutable/Source/MutableValidation/Private/MuV/COMultiPluginBulkValidationCommandlet.cpp
+++ b/Engine/Plugins/Mutable/Source/MutableValidation/Private/MuV/COMultiPluginBulkValidationCommandlet.cpp
@@ -57,7 +57,7 @@ int32 UCOMultiPluginBulkValidationCommandlet::Main(const FString& Params)
 					
 					if (PluginName.IsNone())
 					{
-						UE_LOG(LogMutable, Warning, TEXT("An empty PluginName was provided. Skipping."));
+						UE_LOGF(LogMutable, Warning, "An empty PluginName was provided. Skipping.");
 						continue;
 					}
 					
@@ -95,15 +95,15 @@ int32 UCOMultiPluginBulkValidationCommandlet::Main(const FString& Params)
 			FoundAssetData = FindAllAssetsAtPath(CustomizableObjectsSearchPath, UCustomizableObject::StaticClass());
 			if (FoundAssetData.IsEmpty())
 			{
-				UE_LOG(LogMutable, Display, TEXT("Found no Customizable Objects to validate in the provided path : %s"), *CustomizableObjectsSearchPath.ToString());
+				UE_LOGF(LogMutable, Display, "Found no Customizable Objects to validate in the provided path : %ls", *CustomizableObjectsSearchPath.ToString());
 				continue;
 			}
 		
 			// Log all the Customizable Objects to be tested:
-			UE_LOG(LogMutable, Display, TEXT("Found a total of %u Customizable Objects to validate at '%s'. Some may be discarded based on the test settings."), FoundAssetData.Num(), *CustomizableObjectsSearchPath.ToString())
+			UE_LOGF(LogMutable, Display, "Found a total of %u Customizable Objects to validate at '%ls'. Some may be discarded based on the test settings.", FoundAssetData.Num(), *CustomizableObjectsSearchPath.ToString())
 			for (const FAssetData& MutableAssetData : FoundAssetData)
 			{
-				UE_LOG(LogMutable, Display, TEXT("\t - %s (%s)"), *MutableAssetData.AssetName.ToString(), *MutableAssetData.PackageName.ToString());
+				UE_LOGF(LogMutable, Display, "\t - %ls (%ls)", *MutableAssetData.AssetName.ToString(), *MutableAssetData.PackageName.ToString());
 			}
 		}
 		
@@ -111,7 +111,7 @@ int32 UCOMultiPluginBulkValidationCommandlet::Main(const FString& Params)
 		{
 			if (GIsInitialLoad)
 			{
-				UE_LOG(LogMutable, Warning, TEXT("GC will not run as GIsInitialLoad is currently set to true."));
+				UE_LOGF(LogMutable, Warning, "GC will not run as GIsInitialLoad is currently set to true.");
 			}
 			
 			CollectGarbage(RF_NoFlags, true);
@@ -125,7 +125,7 @@ int32 UCOMultiPluginBulkValidationCommandlet::Main(const FString& Params)
 			// Make all the logs that appear in this scope have a relation with the plugin name provided. This will be later used by the MTU to name the CO generated data
 			const FScopedLogSection PluginLogSection (EMutableLogSection::Plugin, PluginName);
 			
-			UE_LOG(LogMutable, Display, TEXT("Starting testing of COs from the '%s' plugin ."), *PluginName.ToString())
+			UE_LOGF(LogMutable, Display, "Starting testing of COs from the '%ls' plugin .", *PluginName.ToString())
 					
 			// iterate over the FoundAsstData of each Plugin-Dir structure
 			const int32 AssetDataCount = FoundAssetData.Num();
@@ -137,7 +137,7 @@ int32 UCOMultiPluginBulkValidationCommandlet::Main(const FString& Params)
 				TObjectPtr<UObject> FoundObject = UE::Mutable::Private::LoadObject(AssetData);
 				if (!FoundObject)
 				{
-					UE_LOG(LogMutable, Error, TEXT("Failed to load the asset with path : %s ."), *AssetData.GetSoftObjectPath().ToString())
+					UE_LOGF(LogMutable, Error, "Failed to load the asset with path : %ls .", *AssetData.GetSoftObjectPath().ToString())
 					continue;
 				}
 
@@ -146,7 +146,7 @@ int32 UCOMultiPluginBulkValidationCommandlet::Main(const FString& Params)
 				// If we only want to test root objects and this is not one go to the next CO
 				if (bOnlyTestRootObjects && !ICustomizableObjectEditorModule::GetChecked().IsRootObject(*TargetCustomizableObject.Get()))
 				{
-					UE_LOG(LogMutable, Display, TEXT("Skipping CO \"%s\" as it is not a root CO."), *TargetCustomizableObject->GetName())
+					UE_LOGF(LogMutable, Display, "Skipping CO \"%ls\" as it is not a root CO.", *TargetCustomizableObject->GetName())
 					continue;
 				}
 			
@@ -158,19 +158,19 @@ int32 UCOMultiPluginBulkValidationCommandlet::Main(const FString& Params)
 				{
 					if (GIsInitialLoad)
 					{
-						UE_LOG(LogMutable, Warning, TEXT("GC will not run as GIsInitialLoad is currently set to true."));
+						UE_LOGF(LogMutable, Warning, "GC will not run as GIsInitialLoad is currently set to true.");
 					}
 				
 					CollectGarbage(RF_NoFlags, true);
 				}
 
-				UE_LOG(LogMutable, Display, TEXT("%s : Validated %u/%u assets."), *PluginName.ToString(), AssetDataIndex + 1, AssetDataCount);
+				UE_LOGF(LogMutable, Display, "%ls : Validated %u/%u assets.", *PluginName.ToString(), AssetDataIndex + 1, AssetDataCount);
 			}
 			
-			UE_LOG(LogMutable, Display, TEXT("Finished testing the COs from the '%s' plugin ."), *PluginName.ToString())
+			UE_LOGF(LogMutable, Display, "Finished testing the COs from the '%ls' plugin .", *PluginName.ToString())
 		}
 	}
 	
-	UE_LOG(LogMutable, Display, TEXT("Mutable commandlet finished."));
+	UE_LOGF(LogMutable, Display, "Mutable commandlet finished.");
 	return 0;
 }
diff --git a/Engine/Plugins/Mutable/Source/MutableValidation/Private/MuV/CacheAssetRegistryCommandlet.cpp b/Engine/Plugins/Mutable/Source/MutableValidation/Private/MuV/CacheAssetRegistryCommandlet.cpp
--- a/Engine/Plugins/Mutable/Source/MutableValidation/Private/MuV/CacheAssetRegistryCommandlet.cpp
+++ b/Engine/Plugins/Mutable/Source/MutableValidation/Private/MuV/CacheAssetRegistryCommandlet.cpp
@@ -12,7 +12,7 @@ int32 UCacheAssetRegistryCommandlet::Main(const FString& Params)
 	// Ensure we are not running this commandlet in vain as we want the cache to be filled
 	if (FParse::Param(FCommandLine::Get(), TEXT("NoAssetRegistryCacheWrite")))
 	{
-		UE_LOG(LogTemp, Error, TEXT("The Asset Registry data is not going to be cached due to the arg 'NoAssetRegistryCacheWrite' being present in the commandline."));
+		UE_LOGF(LogTemp, Error, "The Asset Registry data is not going to be cached due to the arg 'NoAssetRegistryCacheWrite' being present in the commandline.");
 		return 1;
 	}
 	
diff --git a/Engine/Plugins/Mutable/Source/MutableValidation/Private/MuV/CustomizableObjectBulkValidationCommandlet.cpp b/Engine/Plugins/Mutable/Source/MutableValidation/Private/MuV/CustomizableObjectBulkValidationCommandlet.cpp
--- a/Engine/Plugins/Mutable/Source/MutableValidation/Private/MuV/CustomizableObjectBulkValidationCommandlet.cpp
+++ b/Engine/Plugins/Mutable/Source/MutableValidation/Private/MuV/CustomizableObjectBulkValidationCommandlet.cpp
@@ -48,13 +48,13 @@ int32 UCustomizableObjectBulkValidationCommandlet::Main(const FString& Params)
 	{
 		if (!FParse::Value(*Params, TEXT("CustomizableObjectsSearchPath="), CustomizableObjectsSearchPath))
 		{
-			UE_LOG(LogMutable, Error, TEXT("Failed to parse Customizable Object search path from the provided argument : %s"), *Params)
+			UE_LOGF(LogMutable, Error, "Failed to parse Customizable Object search path from the provided argument : %ls", *Params)
 			return 1;
 		}
 
 		if (CustomizableObjectsSearchPath.IsNone())
 		{
-			UE_LOG(LogMutable, Error, TEXT("The path to scan can not be empty"))
+			UE_LOGF(LogMutable, Error, "The path to scan can not be empty")
 			return 1;
 		}
 	}
@@ -73,15 +73,15 @@ int32 UCustomizableObjectBulkValidationCommandlet::Main(const FString& Params)
 		// Early exit if no instances could be found for testing.
 		if (FoundAssetData.IsEmpty())
 		{
-			UE_LOG(LogMutable, Error, TEXT("Aborting Bulk Customizable Object Validation Test: No assets could be found at the provided package path"));
+			UE_LOGF(LogMutable, Error, "Aborting Bulk Customizable Object Validation Test: No assets could be found at the provided package path");
 			return 1;
 		}
 		
 		// Log all the Customizable Objects to be tested:
-		UE_LOG(LogMutable, Display, TEXT("Found a total of %u Customizable Objects to validate. Some may be discarded based on the test settings."), FoundAssetData.Num())
+		UE_LOGF(LogMutable, Display, "Found a total of %u Customizable Objects to validate. Some may be discarded based on the test settings.", FoundAssetData.Num())
 		for (const FAssetData& MutableAssetData : FoundAssetData)
 		{
-			UE_LOG(LogMutable, Display, TEXT("\t - %s (%s)"), *MutableAssetData.AssetName.ToString(), *MutableAssetData.PackageName.ToString());
+			UE_LOGF(LogMutable, Display, "\t - %ls (%ls)", *MutableAssetData.AssetName.ToString(), *MutableAssetData.PackageName.ToString());
 		}
 	}
 	
@@ -108,7 +108,7 @@ int32 UCustomizableObjectBulkValidationCommandlet::Main(const FString& Params)
 		TObjectPtr<UObject> FoundObject = UE::Mutable::Private::LoadObject(CustomizableObjectAssetData);
 		if (!FoundObject)
 		{
-			UE_LOG(LogMutable, Error, TEXT("Failed to load the asset with path : %s ."), *CustomizableObjectAssetData.GetSoftObjectPath().ToString())
+			UE_LOGF(LogMutable, Error, "Failed to load the asset with path : %ls .", *CustomizableObjectAssetData.GetSoftObjectPath().ToString())
 			continue;
 		}
 
@@ -117,7 +117,7 @@ int32 UCustomizableObjectBulkValidationCommandlet::Main(const FString& Params)
 		// If we only want to test root objects and this is not one go to the next CO
 		if (bOnlyTestRootObjects && !ICustomizableObjectEditorModule::GetChecked().IsRootObject(*TargetCustomizableObject.Get()))
 		{
-			UE_LOG(LogMutable, Display, TEXT("Skipping CO \"%s\" as it is not a root CO."), *TargetCustomizableObject->GetName())
+			UE_LOGF(LogMutable, Display, "Skipping CO \"%ls\" as it is not a root CO.", *TargetCustomizableObject->GetName())
 			continue;
 		}
 		
@@ -129,15 +129,15 @@ int32 UCustomizableObjectBulkValidationCommandlet::Main(const FString& Params)
 		{
 			if (GIsInitialLoad)
 			{
-				UE_LOG(LogMutable, Warning, TEXT("GC will not run as GIsInitialLoad is currently set to true."));
+				UE_LOGF(LogMutable, Warning, "GC will not run as GIsInitialLoad is currently set to true.");
 			}
 			
 			CollectGarbage(RF_NoFlags, true);
 		}
 
-		UE_LOG(LogMutable, Display, TEXT("Validated %u/%u assets."), AssetIndex + 1, FoundAssetsCount);
+		UE_LOGF(LogMutable, Display, "Validated %u/%u assets.", AssetIndex + 1, FoundAssetsCount);
 	}
 
-	UE_LOG(LogMutable, Display, TEXT("Mutable commandlet finished."));
+	UE_LOGF(LogMutable, Display, "Mutable commandlet finished.");
 	return 0;
 }
diff --git a/Engine/Plugins/Mutable/Source/MutableValidation/Private/MuV/CustomizableObjectCompilationUtility.cpp b/Engine/Plugins/Mutable/Source/MutableValidation/Private/MuV/CustomizableObjectCompilationUtility.cpp
--- a/Engine/Plugins/Mutable/Source/MutableValidation/Private/MuV/CustomizableObjectCompilationUtility.cpp
+++ b/Engine/Plugins/Mutable/Source/MutableValidation/Private/MuV/CustomizableObjectCompilationUtility.cpp
@@ -31,19 +31,19 @@ bool FCustomizableObjectCompilationUtility::CompileCustomizableObject(UCustomiza
 	if (bOverrideCompilationOptions)
 	{
 		CompilationOptions = *InCompilationOptionsOverride;
-		UE_LOG(LogMutable,Display,TEXT("CO Compilation options overriden by the user defined ones."));
+		UE_LOGF(LogMutable,Display, "CO Compilation options overriden by the user defined ones.");
 	}
 	else
 	{
 		CompilationOptions = CustomizableObject->GetPrivate()->GetCompileOptions();
-		UE_LOG(LogMutable,Display,TEXT("Compiling CO using it's own compilation options."));
+		UE_LOGF(LogMutable,Display, "Compiling CO using it's own compilation options.");
 	}
 
 	// Ensure that the user has provided a target compilation platform.
 	// Mutable is able to run without one but we want to be explicit in the context of testing.
 	if (!CompilationOptions.TargetPlatform)
 	{
-		UE_LOG(LogMutable, Error, TEXT("The compilation of the %s model could not be started : No explicit platform was provided."), *CustomizableObject->GetName());
+		UE_LOGF(LogMutable, Error, "The compilation of the %ls model could not be started : No explicit platform was provided.", *CustomizableObject->GetName());
 		return false;
 	}
 
@@ -52,10 +52,10 @@ bool FCustomizableObjectCompilationUtility::CompileCustomizableObject(UCustomiza
 	// when we are sure we do not require it.
 	if (bShouldLogMutableLogs)
 	{
-		UE_LOG(LogMutable, Log, TEXT("(string) model_compile_options_overriden : %s "), bOverrideCompilationOptions ? TEXT("true") : TEXT("false"));
-		UE_LOG(LogMutable, Log, TEXT("(int) model_optimization_level : %d "), CompilationOptions.OptimizationLevel);
-		UE_LOG(LogMutable, Log, TEXT("(string) model_texture_compression : %s "), *UEnum::GetValueAsString(CompilationOptions.TextureCompression));
-		UE_LOG(LogMutable, Log, TEXT("(string) model_compile_platform_name : %s "), *CompilationOptions.TargetPlatform->PlatformName());
+		UE_LOGF(LogMutable, Log, "(string) model_compile_options_overriden : %s ", "false");
+		UE_LOGF(LogMutable, Log, "(int) model_optimization_level : %d ", CompilationOptions.OptimizationLevel);
+		UE_LOGF(LogMutable, Log, "(string) model_texture_compression : %ls ", *UEnum::GetValueAsString(CompilationOptions.TextureCompression));
+		UE_LOGF(LogMutable, Log, "(string) model_compile_platform_name : %ls ", *CompilationOptions.TargetPlatform->PlatformName());
 	}
 
 	// Delete data from previous compilations just to be sure it does not affect the new compilation operation
@@ -63,14 +63,14 @@ bool FCustomizableObjectCompilationUtility::CompileCustomizableObject(UCustomiza
 	const FString CompiledDataDirectory = GetCompiledDataFolderPath();
 	if (CustomizableObject->IsCompiled() && FileManager.DirectoryExists(*CompiledDataDirectory))
 	{
-		UE_LOG(LogMutable,Display,TEXT("Deleting old mutable compiled data from '%s'..."), *CompiledDataDirectory);
+		UE_LOGF(LogMutable,Display, "Deleting old mutable compiled data from '%ls'...", *CompiledDataDirectory);
 		if ( !FileManager.DeleteDirectory( *CompiledDataDirectory, false, true) )
 		{
-			UE_LOG(LogMutable,Warning,TEXT("Failed to delete old compiled data directory '%s'. This could affect the compilation of the CO."), *CompiledDataDirectory);
+			UE_LOGF(LogMutable,Warning, "Failed to delete old compiled data directory '%ls'. This could affect the compilation of the CO.", *CompiledDataDirectory);
 		}
 		else
 		{
-			UE_LOG(LogMutable, Display, TEXT("Directory deleted successfully"));
+			UE_LOGF(LogMutable, Display, "Directory deleted successfully");
 		}
 	}
 	
@@ -92,7 +92,7 @@ bool FCustomizableObjectCompilationUtility::CompileCustomizableObject(UCustomiza
 		
 		if (FLowLevelMemTracker::IsEnabled())
 		{
-			UE_LOG(LogMutable,Display,TEXT("LLM system enabled: Peak memory usage during mutable CO compilation will be logged after said compilation takes place."));
+			UE_LOGF(LogMutable,Display, "LLM system enabled: Peak memory usage during mutable CO compilation will be logged after said compilation takes place.");
 			
 			LowLevelMemoryTracker.UpdateStatsPerFrame();
 			CompilationStartTotalBytes = LowLevelMemoryTracker.GetTagAmountForTracker(ELLMTracker::Default, ELLMTag::Total, UE::LLM::ESizeParams::ReportCurrent);
@@ -103,7 +103,7 @@ bool FCustomizableObjectCompilationUtility::CompileCustomizableObject(UCustomiza
 		const int64 CompilationStartMutableBytes = UE::Mutable::Private::FGlobalMemoryCounter::GetAbsoluteCounter();
 		UE::Mutable::Private::FGlobalMemoryCounter::Zero();
 		
-		UE_LOG(LogMutable,Display,TEXT("Compiling Customizable Object..."));
+		UE_LOGF(LogMutable,Display, "Compiling Customizable Object...");
 		
 		TSharedRef<FCompilationRequest> CompileRequest = MakeShared<FCompilationRequest>(*CustomizableObject);
 		CompileRequest->bSilentCompilation = false;
@@ -146,29 +146,29 @@ bool FCustomizableObjectCompilationUtility::CompileCustomizableObject(UCustomiza
 
 		// Report the time we took to run the compilation
 		const double CompilationEndSeconds = FPlatformTime::Seconds() - CompilationStartSeconds;
-		UE_LOG(LogMutable, Display, TEXT("The compilation of the %s CO model took %f seconds."), *CustomizableObject->GetName(), CompilationEndSeconds);
+		UE_LOGF(LogMutable, Display, "The compilation of the %ls CO model took %f seconds.", *CustomizableObject->GetName(), CompilationEndSeconds);
 		if (bShouldLogMutableLogs)
 		{
-			UE_LOG(LogMutable, Log, TEXT("(double) model_compile_time_ms : %f "), CompilationEndSeconds * 1000);
+			UE_LOGF(LogMutable, Log, "(double) model_compile_time_ms : %f ", CompilationEndSeconds * 1000);
 
 			// Also report compilation end status
 			const ECustomizableObjectCompilationState CompilationEndResult = bCompilationSuccess ? ECustomizableObjectCompilationState::Completed
 				: ECustomizableObjectCompilationState::Failed;
-			UE_LOG(LogMutable, Log, TEXT("(string) model_compile_end_state : %s "), *UEnum::GetValueAsString(CompilationEndResult));
+			UE_LOGF(LogMutable, Log, "(string) model_compile_end_state : %ls ", *UEnum::GetValueAsString(CompilationEndResult));
 		}
 		
 		// Get the peak mutable memory used during the compilation operation
 		const int64 CompilationEndPeakMutableBytes = UE::Mutable::Private::FGlobalMemoryCounter::GetPeak();
 		const int64 CompilationEndRealMutablePeakBytes = CompilationStartMutableBytes + CompilationEndPeakMutableBytes;
-		UE_LOG(LogMutable, Display, TEXT("Pre-Compilation Mutable memory usage : %lld"), CompilationStartMutableBytes);
-		UE_LOG(LogMutable, Display, TEXT("Compilation Mutable peak memory usage : %lld"), CompilationEndPeakMutableBytes);
-		UE_LOG(LogMutable, Display, TEXT("Compilation Mutable real peak memory usage : %lld"), CompilationEndRealMutablePeakBytes);
+		UE_LOGF(LogMutable, Display, "Pre-Compilation Mutable memory usage : %lld", CompilationStartMutableBytes);
+		UE_LOGF(LogMutable, Display, "Compilation Mutable peak memory usage : %lld", CompilationEndPeakMutableBytes);
+		UE_LOGF(LogMutable, Display, "Compilation Mutable real peak memory usage : %lld", CompilationEndRealMutablePeakBytes);
 		if (bShouldLogMutableLogs)
 		{
 			// Mutable system reported memory usage
-			UE_LOG(LogMutable, Log, TEXT("(int) model_compilation_start_bytes : %lld "), CompilationStartMutableBytes);
-			UE_LOG(LogMutable, Log, TEXT("(int) model_compilation_end_peak_bytes : %lld "), CompilationEndPeakMutableBytes);
-			UE_LOG(LogMutable, Log, TEXT("(int) model_compilation_end_real_peak_bytes : %lld "), CompilationEndRealMutablePeakBytes);
+			UE_LOGF(LogMutable, Log, "(int) model_compilation_start_bytes : %lld ", CompilationStartMutableBytes);
+			UE_LOGF(LogMutable, Log, "(int) model_compilation_end_peak_bytes : %lld ", CompilationEndPeakMutableBytes);
+			UE_LOGF(LogMutable, Log, "(int) model_compilation_end_real_peak_bytes : %lld ", CompilationEndRealMutablePeakBytes);
 		}
 
 		
@@ -176,18 +176,18 @@ bool FCustomizableObjectCompilationUtility::CompileCustomizableObject(UCustomiza
 #if ENABLE_LOW_LEVEL_MEM_TRACKER
 		if (FLowLevelMemTracker::IsEnabled())
 		{
-			UE_LOG(LogMutable, Display, TEXT("LLM Pre-Compilation Total memory usage (bytes) : %lld"), CompilationStartTotalBytes);
+			UE_LOGF(LogMutable, Display, "LLM Pre-Compilation Total memory usage (bytes) : %lld", CompilationStartTotalBytes);
 			// CompilationPeakBytes : Peak memory recorded during the compilation subtracting from it the memory in use before the compilation itself.
 			// It should reflect the peak memory usage OF the compilation (as direct result of the compilation itself) but it should be used as an approximation
 			// since we can never be too sure no other system is doing work in parallel that is not part of the mutable compilation operation.
 			const int64 CompilationEndTotalPeakBytes = CompilationEndTotalRealPeakBytes - CompilationStartTotalBytes;
-			UE_LOG(LogMutable, Display, TEXT("LLM Total Compilation Peak memory usage (bytes) : %lld"), CompilationEndTotalPeakBytes);
-			UE_LOG(LogMutable, Display, TEXT("LLM Total Compilation Real Peak memory usage (bytes) : %lld"), CompilationEndTotalRealPeakBytes);
+			UE_LOGF(LogMutable, Display, "LLM Total Compilation Peak memory usage (bytes) : %lld", CompilationEndTotalPeakBytes);
+			UE_LOGF(LogMutable, Display, "LLM Total Compilation Real Peak memory usage (bytes) : %lld", CompilationEndTotalRealPeakBytes);
 			if (bShouldLogMutableLogs)
 			{
-				UE_LOG(LogMutable, Log, TEXT("(int) model_compilation_llm_total_start_bytes : %lld "), CompilationStartTotalBytes);
-				UE_LOG(LogMutable, Log, TEXT("(int) model_compilation_llm_total_end_peak_bytes : %lld "), CompilationEndTotalPeakBytes);
-				UE_LOG(LogMutable, Log, TEXT("(int) model_compilation_llm_total_end_real_peak_bytes : %lld "), CompilationEndTotalRealPeakBytes);
+				UE_LOGF(LogMutable, Log, "(int) model_compilation_llm_total_start_bytes : %lld ", CompilationStartTotalBytes);
+				UE_LOGF(LogMutable, Log, "(int) model_compilation_llm_total_end_peak_bytes : %lld ", CompilationEndTotalPeakBytes);
+				UE_LOGF(LogMutable, Log, "(int) model_compilation_llm_total_end_real_peak_bytes : %lld ", CompilationEndTotalRealPeakBytes);
 			}
 		}
 #endif
diff --git a/Engine/Plugins/Mutable/Source/MutableValidation/Private/MuV/CustomizableObjectInstanceUpdateUtility.cpp b/Engine/Plugins/Mutable/Source/MutableValidation/Private/MuV/CustomizableObjectInstanceUpdateUtility.cpp
--- a/Engine/Plugins/Mutable/Source/MutableValidation/Private/MuV/CustomizableObjectInstanceUpdateUtility.cpp
+++ b/Engine/Plugins/Mutable/Source/MutableValidation/Private/MuV/CustomizableObjectInstanceUpdateUtility.cpp
@@ -24,7 +24,7 @@ bool FCustomizableObjectInstanceUpdateUtility::UpdateInstance(UCustomizableObjec
 	
 	// Schedule the update of the COI
 	{
-		UE_LOG(LogMutable, Display, TEXT("Invoking update for %s instance."), *Instance->GetName());
+		UE_LOGF(LogMutable, Display, "Invoking update for %ls instance.", *Instance->GetName());
 
 		// Instance update delegate	
 		FInstanceUpdateNativeDelegate InstanceUpdateNativeDelegate;
@@ -77,7 +77,7 @@ bool FCustomizableObjectInstanceUpdateUtility::UpdateInstance(UCustomizableObjec
 
 			if (bFullyStreamed)
 			{
-				UE_LOG(LogMutable, Display, TEXT("Instance %s finished streaming all MIPs."), *Instance->GetName());
+				UE_LOGF(LogMutable, Display, "Instance %ls finished streaming all MIPs.", *Instance->GetName());
 				ComponentsBeingUpdated.Reset();
 				
 				bIsInstanceBeingUpdated = false;		// Exit the while loop
@@ -99,7 +99,7 @@ void FCustomizableObjectInstanceUpdateUtility::OnInstanceUpdateResult(const FUpd
 	
 	if (UCustomizableObjectSystem::IsUpdateResultValid(Result.UpdateResult))
 	{
-		UE_LOG(LogMutable, Display, TEXT("Instance %s finished update successfully."), *InstanceName);
+		UE_LOGF(LogMutable, Display, "Instance %ls finished update successfully.", *InstanceName);
 		bInstanceFailedUpdate = false;
 		
 		// Request load all MIPs
@@ -132,12 +132,12 @@ void FCustomizableObjectInstanceUpdateUtility::OnInstanceUpdateResult(const FUpd
 			ComponentsBeingUpdated.Add(SkeletalComponent);
 		}
 
-		UE_LOG(LogMutable, Display, TEXT("Instance %s requesting the streaming of all MIPs."), *InstanceName);
+		UE_LOGF(LogMutable, Display, "Instance %ls requesting the streaming of all MIPs.", *InstanceName);
 	}
 	else
 	{
 		const FString OutputStatus = UEnum::GetValueAsString(Result.UpdateResult);
-		UE_LOG(LogMutable, Error, TEXT("Instance %s finished update with anomalous state : %s."), *InstanceName, *OutputStatus);
+		UE_LOGF(LogMutable, Error, "Instance %ls finished update with anomalous state : %ls.", *InstanceName, *OutputStatus);
 		bInstanceFailedUpdate = true;
 		
 		// Tell the system the instance finished its update so we can continue the execution without waiting for the mips to stream in
diff --git a/Engine/Plugins/Mutable/Source/MutableValidation/Private/MuV/CustomizableObjectValidationCommandlet.cpp b/Engine/Plugins/Mutable/Source/MutableValidation/Private/MuV/CustomizableObjectValidationCommandlet.cpp
--- a/Engine/Plugins/Mutable/Source/MutableValidation/Private/MuV/CustomizableObjectValidationCommandlet.cpp
+++ b/Engine/Plugins/Mutable/Source/MutableValidation/Private/MuV/CustomizableObjectValidationCommandlet.cpp
@@ -38,7 +38,7 @@ int32 UCustomizableObjectValidationCommandlet::Main(const FString& Params)
 		FString CustomizableObjectAssetPath = "";
 		if (!FParse::Value(*Params, TEXT("CustomizableObject="), CustomizableObjectAssetPath))
 		{
-			UE_LOG(LogMutable, Error, TEXT("Failed to parse Customizable Object package name from provided argument : %s"), *Params)
+			UE_LOGF(LogMutable, Error, "Failed to parse Customizable Object package name from provided argument : %ls", *Params)
 			return 1;
 		}
 	
@@ -46,7 +46,7 @@ int32 UCustomizableObjectValidationCommandlet::Main(const FString& Params)
 		UObject* FoundObject = UE::Mutable::Private::LoadObject(FSoftObjectPath(CustomizableObjectAssetPath));
 		if (!FoundObject)
 		{
-			UE_LOG(LogMutable, Error, TEXT("Failed to retrieve UObject from path %s"), *CustomizableObjectAssetPath);
+			UE_LOGF(LogMutable, Error, "Failed to retrieve UObject from path %ls", *CustomizableObjectAssetPath);
 			return 1;
 		}
 	
@@ -54,7 +54,7 @@ int32 UCustomizableObjectValidationCommandlet::Main(const FString& Params)
 		TargetCustomizableObject = TStrongObjectPtr(Cast<UCustomizableObject>(FoundObject));
 		if (!TargetCustomizableObject)
 		{
-			UE_LOG(LogMutable, Error, TEXT("Failed to cast found UObject to UCustomizableObject."));
+			UE_LOGF(LogMutable, Error, "Failed to cast found UObject to UCustomizableObject.");
 			return 1;
 		}
 	}
@@ -87,6 +87,6 @@ int32 UCustomizableObjectValidationCommandlet::Main(const FString& Params)
 	TestCustomizableObject(*TargetCustomizableObject, *TargetCompilationPlatform, InstancesToGenerate);
 	// -----------------------------------------------------------------------------------------------------------------------------------------------
 	
-	UE_LOG(LogMutable, Display, TEXT("Mutable commandlet finished."));
+	UE_LOGF(LogMutable, Display, "Mutable commandlet finished.");
 	return 0;
 }
diff --git a/Engine/Plugins/Mutable/Source/MutableValidation/Private/MuV/ScopedLogSection.cpp b/Engine/Plugins/Mutable/Source/MutableValidation/Private/MuV/ScopedLogSection.cpp
--- a/Engine/Plugins/Mutable/Source/MutableValidation/Private/MuV/ScopedLogSection.cpp
+++ b/Engine/Plugins/Mutable/Source/MutableValidation/Private/MuV/ScopedLogSection.cpp
@@ -20,7 +20,7 @@ FScopedLogSection::FScopedLogSection(EMutableLogSection Section, const FName& In
 		check(!InSectionTarget.IsNone());
 		SectionObject = InSectionTarget;
 		
-		UE_LOG(LogMutableValidation, Log, TEXT(" SECTION START : %s - [%s] "), *GetLogSectionName(CurrentSection), *SectionObject.ToString());
+		UE_LOGF(LogMutableValidation, Log, " SECTION START : %ls - [%ls] ", *GetLogSectionName(CurrentSection), *SectionObject.ToString());
 	}
 	else if (CurrentSection == EMutableLogSection::Plugin)
 	{
@@ -29,11 +29,11 @@ FScopedLogSection::FScopedLogSection(EMutableLogSection Section, const FName& In
 		check(!InSectionTarget.IsNone());
 		SectionPlugin = InSectionTarget;
 		
-		UE_LOG(LogMutableValidation, Log, TEXT(" SECTION START : %s - [%s] "), *GetLogSectionName(CurrentSection), *SectionPlugin.ToString());
+		UE_LOGF(LogMutableValidation, Log, " SECTION START : %ls - [%ls] ", *GetLogSectionName(CurrentSection), *SectionPlugin.ToString());
 	}
 	else
 	{
-		UE_LOG(LogMutableValidation, Log, TEXT(" SECTION START : %s "), *GetLogSectionName(CurrentSection));
+		UE_LOGF(LogMutableValidation, Log, " SECTION START : %ls ", *GetLogSectionName(CurrentSection));
 	}
 }
 
@@ -42,17 +42,17 @@ FScopedLogSection::~FScopedLogSection()
 {
 	if (CurrentSection == EMutableLogSection::Object)
 	{
-		UE_LOG(LogMutableValidation, Log, TEXT(" SECTION END : %s - [%s] "), *GetLogSectionName(CurrentSection), *SectionObject.ToString());
+		UE_LOGF(LogMutableValidation, Log, " SECTION END : %ls - [%ls] ", *GetLogSectionName(CurrentSection), *SectionObject.ToString());
 		SectionObject = NAME_None;
 	}
 	else if (CurrentSection == EMutableLogSection::Plugin)
 	{
-		UE_LOG(LogMutableValidation, Log, TEXT(" SECTION END : %s - [%s] "), *GetLogSectionName(CurrentSection), *SectionPlugin.ToString());
+		UE_LOGF(LogMutableValidation, Log, " SECTION END : %ls - [%ls] ", *GetLogSectionName(CurrentSection), *SectionPlugin.ToString());
 		SectionPlugin = NAME_None;
 	}
 	else
 	{
-		UE_LOG(LogMutableValidation, Log, TEXT(" SECTION END : %s "), *GetLogSectionName(CurrentSection));
+		UE_LOGF(LogMutableValidation, Log, " SECTION END : %ls ", *GetLogSectionName(CurrentSection));
 	}
 	
 	// Set the current section to none (undefined)
diff --git a/Engine/Plugins/Mutable/Source/MutableValidation/Private/MuV/ValidationUtils.cpp b/Engine/Plugins/Mutable/Source/MutableValidation/Private/MuV/ValidationUtils.cpp
--- a/Engine/Plugins/Mutable/Source/MutableValidation/Private/MuV/ValidationUtils.cpp
+++ b/Engine/Plugins/Mutable/Source/MutableValidation/Private/MuV/ValidationUtils.cpp
@@ -22,14 +22,14 @@
 void PrepareAssetRegistry()
 {
 	FAssetRegistryModule& AssetRegistryModule = FModuleManager::LoadModuleChecked<FAssetRegistryModule>(AssetRegistryConstants::ModuleName);
-	UE_LOG(LogMutable,Display,TEXT("Searching all assets (this will take some time)..."));
+	UE_LOGF(LogMutable,Display, "Searching all assets (this will take some time)...");
 	
 	const double AssetRegistrySearchStartSeconds = FPlatformTime::Seconds();
 	AssetRegistryModule.Get().SearchAllAssets(true /* bSynchronousSearch */);
 	const double AssetRegistrySearchEndSeconds = FPlatformTime::Seconds() - AssetRegistrySearchStartSeconds;
-	UE_LOG(LogMutable, Log, TEXT("(double) asset_registry_search_time_s : %f "), AssetRegistrySearchEndSeconds);
+	UE_LOGF(LogMutable, Log, "(double) asset_registry_search_time_s : %f ", AssetRegistrySearchEndSeconds);
 
-	UE_LOG(LogMutable,Display,TEXT("Asset searching completed in \"%f\" seconds!"), AssetRegistrySearchEndSeconds);
+	UE_LOGF(LogMutable,Display, "Asset searching completed in \"%f\" seconds!", AssetRegistrySearchEndSeconds);
 }
 
 
@@ -37,13 +37,13 @@ void LogGlobalSettings()
 {
 	// Mutable Settings
 	const int32 WorkingMemoryKB = UCustomizableObjectSystem::GetInstanceChecked()->GetWorkingMemory() ;
-	UE_LOG(LogMutable,Log, TEXT("(int) working_memory_bytes : %d"), WorkingMemoryKB*1024)
-	UE_LOG(LogMutable, Display, TEXT("The mutable updates will use as working memory the value of %d KB"), WorkingMemoryKB)
+	UE_LOGF(LogMutable,Log, "(int) working_memory_bytes : %d", WorkingMemoryKB*1024)
+	UE_LOGF(LogMutable, Display, "The mutable updates will use as working memory the value of %d KB", WorkingMemoryKB)
 	
 	// Expand this when adding new controls from the .xml file
 	
 	// RHI Settings
-	UE_LOG(LogMutable, Log, TEXT("(string) rhi_adapter_name : %s"), *GRHIAdapterName )
+	UE_LOGF(LogMutable, Log, "(string) rhi_adapter_name : %ls", *GRHIAdapterName )
 }
 
 
@@ -51,7 +51,7 @@ void Wait(const double ToWaitSeconds)
 {
 	check (ToWaitSeconds > 0);
 	
-	UE_LOG(LogMutable,Display,TEXT("Holding test execution for %f seconds."),ToWaitSeconds);
+	UE_LOGF(LogMutable,Display, "Holding test execution for %f seconds.",ToWaitSeconds);
 	const double EndSeconds = FPlatformTime::Seconds() + ToWaitSeconds;
 	while (FPlatformTime::Seconds() < EndSeconds)
 	{
@@ -65,7 +65,7 @@ void Wait(const double ToWaitSeconds)
 		}
 	}
 
-	UE_LOG(LogMutable,Display,TEXT("Resuming test execution."));
+	UE_LOGF(LogMutable,Display, "Resuming test execution.");
 }
 
 
@@ -95,13 +95,13 @@ TArray<FAssetData> FindAllAssetsAtPath(FName SearchPath, const UClass* TargetObj
 		// Ensure the AR module is ready to search for stuff
 		AssetRegistryModule.Get().SearchAllAssets(true);
 	
-		UE_LOG(LogMutable, Display, TEXT("Searching for all %s objects to test at path : %s ."), *TargetObjectClass->GetName(), *SearchPath.ToString());
+		UE_LOGF(LogMutable, Display, "Searching for all %ls objects to test at path : %ls .", *TargetObjectClass->GetName(), *SearchPath.ToString());
 		AssetRegistryModule.Get().GetAssets(Filter, FoundAssetData);
-		UE_LOG(LogMutable, Display, TEXT("Search of %s objects completed. Found %i objects."), *TargetObjectClass->GetName(), FoundAssetData.Num());
+		UE_LOGF(LogMutable, Display, "Search of %ls objects completed. Found %i objects.", *TargetObjectClass->GetName(), FoundAssetData.Num());
 	}
 	else
 	{
-		UE_LOG(LogMutable, Error, TEXT("No objects can be retrieved using a null class."));
+		UE_LOGF(LogMutable, Error, "No objects can be retrieved using a null class.");
 	}
 
 	return FoundAssetData;
@@ -114,7 +114,7 @@ uint32 GetTargetAmountOfInstances(const FString& Params)
 	uint32 InstancesToGenerate = 16;
 	if (!FParse::Value(*Params, TEXT("InstanceGenerationCount="),InstancesToGenerate))
 	{
-		UE_LOG(LogMutable, Display, TEXT("Instance generation count not specified. Using default value : %u"), InstancesToGenerate);
+		UE_LOGF(LogMutable, Display, "Instance generation count not specified. Using default value : %u", InstancesToGenerate);
 	}
 	return InstancesToGenerate;
 }
@@ -126,7 +126,7 @@ ITargetPlatform* GetCompilationPlatform(const FString& Params)
 	FString TargetPlatformName = "";
 	if (!FParse::Value(*Params, TEXT("CompilationPlatformName="), TargetPlatformName))
 	{
-		UE_LOG(LogMutable, Error, TEXT("Failed to parse the target compilation platform. Have you even provided the argument?"))
+		UE_LOGF(LogMutable, Error, "Failed to parse the target compilation platform. Have you even provided the argument?")
 		return nullptr;
 	}
 
@@ -149,7 +149,7 @@ ITargetPlatform* GetCompilationPlatform(const FString& Params)
 	
 	if (!TargetCompilationPlatform)
 	{
-		UE_LOG(LogMutable, Error, TEXT("Unable to relate the provided platform name (%s) with the available platforms in this machine."), *TargetPlatformName);
+		UE_LOGF(LogMutable, Error, "Unable to relate the provided platform name (%ls) with the available platforms in this machine.", *TargetPlatformName);
 	}
 
 	return TargetCompilationPlatform;
@@ -162,7 +162,7 @@ bool GetIfOnlyRootCOsShouldBeUsed(const FString& Params)
 	FParse::Bool(*Params,TEXT("SkipNonRootObjects="),bOnlyTestRootObjects);
 	if (bOnlyTestRootObjects)
 	{
-		UE_LOG(LogMutable, Display, TEXT("Only the root COs will be tested")) ;
+		UE_LOGF(LogMutable, Display, "Only the root COs will be tested") ;
 	}
 	
 	return bOnlyTestRootObjects;
@@ -175,7 +175,7 @@ bool TestCustomizableObject(UCustomizableObject& InTargetCustomizableObject, con
 	const FName ObjectName = FName(InTargetCustomizableObject.GetPathName());
 	if (ObjectName.IsNone())
 	{
-		UE_LOG(LogMutable, Error, TEXT("The compilation of the Customizable object was not successful : The name of the CO was found to be None."));
+		UE_LOGF(LogMutable, Error, "The compilation of the Customizable object was not successful : The name of the CO was found to be None.");
 		return false;
 	}
 	
@@ -201,7 +201,7 @@ bool TestCustomizableObject(UCustomizableObject& InTargetCustomizableObject, con
 
 	if (!bWasCoCompilationSuccessful)
 	{
-		UE_LOG(LogMutable, Error, TEXT("The compilation of the Customizable object was not successful : No instances will be generated."));
+		UE_LOGF(LogMutable, Error, "The compilation of the Customizable object was not successful : No instances will be generated.");
 		return false;		// Validation failed
 	}
 	
@@ -221,8 +221,8 @@ bool TestCustomizableObject(UCustomizableObject& InTargetCustomizableObject, con
 			}
 
 			// Print MTU parseable logs
-			UE_LOG(LogMutable, Log, TEXT("(int) model_rom_count : %d "), RomCount);
-			UE_LOG(LogMutable, Log, TEXT("(int) model_roms_size : %lld "), TotalRomSizeBytes);
+			UE_LOGF(LogMutable, Log, "(int) model_rom_count : %d ", RomCount);
+			UE_LOGF(LogMutable, Log, "(int) model_roms_size : %lld ", TotalRomSizeBytes);
 		}
 
 		// CO embedded data size ------ //
@@ -233,21 +233,21 @@ bool TestCustomizableObject(UCustomizableObject& InTargetCustomizableObject, con
 			InTargetCustomizableObject.GetPrivate()->SaveEmbeddedData(SerializationTarget);
 			const int64 COEmbeddedDataSizeBytes = EmbeddedDataBytes.Num();
 		
-			UE_LOG(LogMutable, Log, TEXT("(int) co_embedded_data_bytes : %lld "), COEmbeddedDataSizeBytes);
+			UE_LOGF(LogMutable, Log, "(int) co_embedded_data_bytes : %lld ", COEmbeddedDataSizeBytes);
 		}
 	}
 	
 	// Skip instances updating if no instances should be updated 
 	if (InstancesToGenerate <= 0)
 	{
-		UE_LOG(LogMutable, Display, TEXT("Instances to generate are 0 : No instances will be generated."));
+		UE_LOGF(LogMutable, Display, "Instances to generate are 0 : No instances will be generated.");
 		return true;	// No instances are targeted for generation, this will be taken as compilation only test.
 	}
 
 	// Do not generate instances if the selected platform is not the running platform
 	if (&TargetCompilationPlatform != GetTargetPlatformManagerRef().GetRunningTargetPlatform())
 	{
-		UE_LOG(LogMutable, Display, TEXT("RunningPlatform != UserProvidedCompilationPlatform : No instances will be generated."));
+		UE_LOGF(LogMutable, Display, "RunningPlatform != UserProvidedCompilationPlatform : No instances will be generated.");
 		return true;
 	}
 
@@ -265,10 +265,10 @@ bool TestCustomizableObject(UCustomizableObject& InTargetCustomizableObject, con
 	}
 	// ---------------------------------------------------------------------------------------------------------- //
 
-	UE_LOG(LogMutable, Log, TEXT("(int) generated_instances_count : %u "), GeneratedInstances);
+	UE_LOGF(LogMutable, Log, "(int) generated_instances_count : %u ", GeneratedInstances);
 	
 	// Update the instances generated --------------------------------------------------------------------------- //
-	UE_LOG(LogMutable, Display, TEXT("Updating generated instances..."));
+	UE_LOGF(LogMutable, Display, "Updating generated instances...");
 	bool bInstanceFailedUpdate = false;
 	const double InstancesUpdateStartSeconds = FPlatformTime::Seconds();
 	{
@@ -292,24 +292,24 @@ bool TestCustomizableObject(UCustomizableObject& InTargetCustomizableObject, con
 	
 	// Notify and log time required by the instances to get updated
 	const double CombinedInstanceUpdateSeconds = InstancesUpdateEndSeconds - InstancesUpdateStartSeconds;
-	UE_LOG(LogMutable, Log, TEXT("(double) combined_update_time_ms : %f "), CombinedInstanceUpdateSeconds * 1000);
+	UE_LOGF(LogMutable, Log, "(double) combined_update_time_ms : %f ", CombinedInstanceUpdateSeconds * 1000);
 
 	check(GeneratedInstances > 0);
 	const double AverageInstanceUpdateSeconds = CombinedInstanceUpdateSeconds / GeneratedInstances;
-	UE_LOG(LogMutable, Log, TEXT("(double) avg_update_time_ms : %f "), AverageInstanceUpdateSeconds * 1000);
+	UE_LOGF(LogMutable, Log, "(double) avg_update_time_ms : %f ", AverageInstanceUpdateSeconds * 1000);
 
-	UE_LOG(LogMutable, Display, TEXT("Generation of Customizable object instances took %f seconds (%f seconds avg)."), CombinedInstanceUpdateSeconds, AverageInstanceUpdateSeconds);
+	UE_LOGF(LogMutable, Display, "Generation of Customizable object instances took %f seconds (%f seconds avg).", CombinedInstanceUpdateSeconds, AverageInstanceUpdateSeconds);
 	// ---------------------------------------------------------------------------------------------------------- //
 
 	// Compute instance update result
 	const bool bInstancesTestedSuccessfully = !bInstanceFailedUpdate && bWasInstancesCreationSuccessful;
 	if (bInstancesTestedSuccessfully)
     {
-        UE_LOG(LogMutable, Display, TEXT("Generation of Customizable object instances was successful."));
+        UE_LOGF(LogMutable, Display, "Generation of Customizable object instances was successful.");
     }
     else
     {
-        UE_LOG(LogMutable, Error, TEXT("The generation of Customizable object instances was not successful."));
+        UE_LOGF(LogMutable, Error, "The generation of Customizable object instances was not successful.");
     }
 	
 	return bInstancesTestedSuccessfully;
