[Mutable] Change all instances of UE_LOG to UE_LOGF in CustomizableObject module.

#rnx
#Synced-CL 51758072

diff --git a/Engine/Plugins/Mutable/Source/CustomizableObject/Private/MuCO/CustomizableObject.cpp b/Engine/Plugins/Mutable/Source/CustomizableObject/Private/MuCO/CustomizableObject.cpp
--- a/Engine/Plugins/Mutable/Source/CustomizableObject/Private/MuCO/CustomizableObject.cpp
+++ b/Engine/Plugins/Mutable/Source/CustomizableObject/Private/MuCO/CustomizableObject.cpp
@@ -184,7 +184,7 @@ void UCustomizableObject::PreSave(FObjectPreSaveContext ObjectSaveContext)
 		}
 		else
 		{
-			UE_LOG(LogMutable, Warning, TEXT("Cook: Customizable Object [%s] is missing [%s] platform data."), *GetName(),
+			UE_LOGF(LogMutable, Warning, "Cook: Customizable Object [%ls] is missing [%ls] platform data.", *GetName(),
 				*ObjectSaveContext.GetTargetPlatform()->PlatformName());
 			
 			// Clear model resources
@@ -275,7 +275,7 @@ void UCustomizableObject::PostLoad()
 			FMessageDialog::Open(EAppMsgType::Ok, LOCTEXT("OldCustomizableObject", "Old Customizable Objects found. Please resave them to avoid future incompatibilities. See output log for more information."));
 		}
 
-		UE_LOG(LogMutable, Warning, TEXT("Unsupported old Customizable Object. Please resave it: %s"), *GetPackage()->GetPathName());
+		UE_LOGF(LogMutable, Warning, "Unsupported old Customizable Object. Please resave it: %ls", *GetPackage()->GetPathName());
 	}
 	
 #if WITH_EDITOR
@@ -385,7 +385,7 @@ void UCustomizableObject::Serialize(FArchive& Ar_Asset)
 	{
 		if (Ar_Asset.IsSaving())
 		{
-			UE_LOG(LogMutable, Verbose, TEXT("Serializing cooked data for Customizable Object [%s]."), *GetName());
+			UE_LOGF(LogMutable, Verbose, "Serializing cooked data for Customizable Object [%ls].", *GetName());
 			GetPrivate()->SaveEmbeddedData(Ar_Asset);
 		}
 	}
@@ -595,7 +595,7 @@ void UCustomizableObjectPrivate::LoadCompiledDataFromDisk()
 			{
 				if (OutOfDatePackages.Num())
 				{
-					UE_LOG(LogMutable, Display, TEXT("Invalidating compiled data due to changes in %s."), *OutOfDatePackages[0].ToString());
+					UE_LOGF(LogMutable, Display, "Invalidating compiled data due to changes in %ls.", *OutOfDatePackages[0].ToString());
 				}
 
 				PrintParticipatingPackagesDiff(OutOfDatePackages, AddedPackages, RemovedPackages, bReleaseVersion);
@@ -768,7 +768,7 @@ FString UCustomizableObject::GetDesc()
 
 void UCustomizableObjectPrivate::SaveEmbeddedData(FArchive& Ar)
 {
-	UE_LOG(LogMutable, Verbose, TEXT("Saving embedded data for Customizable Object [%s] now at position %d."), *GetName(), int(Ar.Tell()));
+	UE_LOGF(LogMutable, Verbose, "Saving embedded data for Customizable Object [%ls] now at position %d.", *GetName(), int(Ar.Tell()));
 
 	TSharedPtr<UE::Mutable::Private::FModel> Model = GetModel();
 
@@ -785,7 +785,7 @@ void UCustomizableObjectPrivate::SaveEmbeddedData(FArchive& Ar)
 		UE::Mutable::Private::FOutputArchive Arch(&Stream);
 		UE::Mutable::Private::FModel::Serialise(Model.Get(), Arch);
 
-		UE_LOG(LogMutable, Verbose, TEXT("Saved embedded data for Customizable Object [%s] now at position %d."), *GetName(), int(Ar.Tell()));
+		UE_LOGF(LogMutable, Verbose, "Saved embedded data for Customizable Object [%ls] now at position %d.", *GetName(), int(Ar.Tell()));
 	}
 }
 
@@ -998,7 +998,7 @@ UCustomizableObjectInstance* UCustomizableObject::CreateInstance()
 	PreviewInstance->SetObject(this);
 	PreviewInstance->GetPrivate()->bShowOnlyRuntimeParameters = false;
 
-	UE_LOG(LogMutable, Verbose, TEXT("Created Customizable Object Instance."));
+	UE_LOGF(LogMutable, Verbose, "Created Customizable Object Instance.");
 
 	return PreviewInstance;
 }
@@ -1106,12 +1106,12 @@ const FString& UCustomizableObjectPrivate::GetIntParameterAvailableOption(int32
 		}
 		else
 		{
-			UE_LOG(LogMutable, Warning, TEXT("Index [%d] out of IntParameterNumOptions bounds at GetIntParameterAvailableOption at CO %s."), K, *GetName());
+			UE_LOGF(LogMutable, Warning, "Index [%d] out of IntParameterNumOptions bounds at GetIntParameterAvailableOption at CO %ls.", K, *GetName());
 		}
 	}
 	else
 	{
-		UE_LOG(LogMutable, Warning, TEXT("Index [%d] out of ParameterProperties bounds at GetIntParameterAvailableOption at CO %s."), ParamIndex, *GetName());
+		UE_LOGF(LogMutable, Warning, "Index [%d] out of ParameterProperties bounds at GetIntParameterAvailableOption at CO %ls.", ParamIndex, *GetName());
 	}
 
 	return s_EmptyString;
@@ -1126,7 +1126,7 @@ int32 UCustomizableObjectPrivate::GetEnumParameterNumValues(int32 ParamIndex) co
 	}
 	else
 	{
-		UE_LOG(LogMutable, Warning, TEXT("Index [%d] out of ParameterProperties bounds at GetIntParameterNumOptions at CO %s."), ParamIndex, *GetName());
+		UE_LOGF(LogMutable, Warning, "Index [%d] out of ParameterProperties bounds at GetIntParameterNumOptions at CO %ls.", ParamIndex, *GetName());
 	}
 
 	return 0;
@@ -1149,7 +1149,7 @@ FString UCustomizableObjectPrivate::FindIntParameterValueName(int32 ParamIndex,
 	}
 	else
 	{
-		UE_LOG(LogMutable, Warning, TEXT("Index [%d] out of ParameterProperties bounds at FindIntParameterValueName at CO %s."), ParamIndex, *GetName());
+		UE_LOGF(LogMutable, Warning, "Index [%d] out of ParameterProperties bounds at FindIntParameterValueName at CO %ls.", ParamIndex, *GetName());
 	}
 
 	return FString();
@@ -1201,7 +1201,7 @@ EMutableParameterType UCustomizableObjectPrivate::GetParameterType(int32 ParamIn
 	}
 	else
 	{
-		UE_LOG(LogMutable, Error, TEXT("Index [%d] out of ParameterProperties bounds at GetParameterType."), ParamIndex);
+		UE_LOGF(LogMutable, Error, "Index [%d] out of ParameterProperties bounds at GetParameterType.", ParamIndex);
 	}
 
 	return EMutableParameterType::None;
@@ -1216,7 +1216,7 @@ EMutableParameterType UCustomizableObject::GetParameterTypeByName(const FString&
 		return GetPrivate()->ParameterProperties[Index].Type;
 	}
 
-	UE_LOG(LogMutable, Warning, TEXT("Name '%s' does not exist in ParameterProperties lookup table at GetParameterTypeByName at CO %s."), *Name, *GetName());
+	UE_LOGF(LogMutable, Warning, "Name '%ls' does not exist in ParameterProperties lookup table at GetParameterTypeByName at CO %ls.", *Name, *GetName());
 
 	for (int32 ParamIndex = 0; ParamIndex < GetPrivate()->ParameterProperties.Num(); ++ParamIndex)
 	{
@@ -1226,7 +1226,7 @@ EMutableParameterType UCustomizableObject::GetParameterTypeByName(const FString&
 		}
 	}
 
-	UE_LOG(LogMutable, Warning, TEXT("Name '%s' does not exist in ParameterProperties at GetParameterTypeByName at CO %s."), *Name, *GetName());
+	UE_LOGF(LogMutable, Warning, "Name '%ls' does not exist in ParameterProperties at GetParameterTypeByName at CO %ls.", *Name, *GetName());
 
 	return EMutableParameterType::None;
 }
@@ -1240,7 +1240,7 @@ const FString& UCustomizableObject::GetParameterName(int32 ParamIndex) const
 	}
 	else
 	{
-		UE_LOG(LogMutable, Warning, TEXT("Index [%d] out of ParameterProperties bounds at GetParameterName at CO %s."), ParamIndex, *GetName());
+		UE_LOGF(LogMutable, Warning, "Index [%d] out of ParameterProperties bounds at GetParameterName at CO %ls.", ParamIndex, *GetName());
 	}
 
 	return s_EmptyString;
@@ -1560,7 +1560,7 @@ float UCustomizableObject::GetFloatParameterDefaultValue(const FString& InParame
 	const int32 ParameterIndex = GetPrivate()->FindParameter(InParameterName);
 	if (ParameterIndex == INDEX_NONE)
 	{
-		UE_LOG(LogMutable, Error, TEXT("Tried to access the default value of the nonexistent float parameter [%s] in the CustomizableObject [%s]."), *InParameterName, *GetName());
+		UE_LOGF(LogMutable, Error, "Tried to access the default value of the nonexistent float parameter [%ls] in the CustomizableObject [%ls].", *InParameterName, *GetName());
 		return FCustomizableObjectFloatParameterValue::DEFAULT_PARAMETER_VALUE;
 	}
 
@@ -1580,7 +1580,7 @@ int32 UCustomizableObject::GetEnumParameterDefaultValue(const FString& InParamet
 	const int32 ParameterIndex = GetPrivate()->FindParameter(InParameterName);
 	if (ParameterIndex == INDEX_NONE)
 	{
-		UE_LOG(LogMutable, Error, TEXT("Tried to access the default value of the nonexistent integer parameter [%s] in the CustomizableObject [%s]."), *InParameterName, *GetName());
+		UE_LOGF(LogMutable, Error, "Tried to access the default value of the nonexistent integer parameter [%ls] in the CustomizableObject [%ls].", *InParameterName, *GetName());
 		return FCustomizableObjectIntParameterValue::DEFAULT_PARAMETER_VALUE;
 	}
 
@@ -1600,7 +1600,7 @@ bool UCustomizableObject::GetBoolParameterDefaultValue(const FString& InParamete
 	const int32 ParameterIndex = GetPrivate()->FindParameter(InParameterName);
 	if (ParameterIndex == INDEX_NONE)
 	{
-		UE_LOG(LogMutable, Error, TEXT("Tried to access the default value of the nonexistent boolean parameter [%s] in the CustomizableObject [%s]."), *InParameterName, *GetName());
+		UE_LOGF(LogMutable, Error, "Tried to access the default value of the nonexistent boolean parameter [%ls] in the CustomizableObject [%ls].", *InParameterName, *GetName());
 		return FCustomizableObjectBoolParameterValue::DEFAULT_PARAMETER_VALUE;
 	}
 
@@ -1626,7 +1626,7 @@ FVector4f UCustomizableObject::GetVectorParameterDefaultValue(const FString& Par
 	const int32 ParameterIndex = GetPrivate()->FindParameter(ParameterName);
 	if (ParameterIndex == INDEX_NONE)
 	{
-		UE_LOG(LogMutable, Error, TEXT("Tried to access the default value of the nonexistent color parameter [%s] in the CustomizableObject [%s]."), *ParameterName, *GetName());
+		UE_LOGF(LogMutable, Error, "Tried to access the default value of the nonexistent color parameter [%ls] in the CustomizableObject [%ls].", *ParameterName, *GetName());
 		return FCustomizableObjectVectorParameterValue::DEFAULT_PARAMETER_VALUE;
 	}
 
@@ -1649,7 +1649,7 @@ FTransform UCustomizableObject::GetTransformParameterDefaultValue(const FString&
 	const int32 ParameterIndex = GetPrivate()->FindParameter(InParameterName);
 	if (ParameterIndex == INDEX_NONE)
 	{
-		UE_LOG(LogMutable, Error, TEXT("Tried to access the default value of the nonexistent color parameter [%s] in the CustomizableObject [%s]."), *InParameterName, *GetName());
+		UE_LOGF(LogMutable, Error, "Tried to access the default value of the nonexistent color parameter [%ls] in the CustomizableObject [%ls].", *InParameterName, *GetName());
 		return FCustomizableObjectTransformParameterValue::DEFAULT_PARAMETER_VALUE;
 	}
 
@@ -1671,7 +1671,7 @@ FCustomizableObjectProjector UCustomizableObject::GetProjectorParameterDefaultVa
 	const int32 ParameterIndex = GetPrivate()->FindParameter(InParameterName);
 	if (ParameterIndex == INDEX_NONE)
 	{
-		UE_LOG(LogMutable, Error, TEXT("Tried to access the default value of the nonexistent projector [%s] in the CustomizableObject [%s]."), *InParameterName, *GetName());
+		UE_LOGF(LogMutable, Error, "Tried to access the default value of the nonexistent projector [%ls] in the CustomizableObject [%ls].", *InParameterName, *GetName());
 		return FCustomizableObjectProjectorParameterValue::DEFAULT_PARAMETER_VALUE;
 	}
 
@@ -1738,7 +1738,7 @@ bool UCustomizableObject::IsParameterMultidimensional(const FString& InParameter
 	const int32 ParameterIndex = GetPrivate()->FindParameter(InParameterName);
 	if (ParameterIndex == INDEX_NONE)
 	{
-		UE_LOG(LogMutable, Error, TEXT("Tried to access the default value of the nonexistent parameter [%s] in the CustomizableObject [%s]."), *InParameterName, *GetName());
+		UE_LOGF(LogMutable, Error, "Tried to access the default value of the nonexistent parameter [%ls] in the CustomizableObject [%ls].", *InParameterName, *GetName());
 		return false;
 	}
 
@@ -2522,7 +2522,7 @@ namespace UE::Mutable::Private
 						[Category](const FFileCategoryOverride& C) { return C.Id == Category.Id; });
 					if (ensure(CategoryOverride))
 					{
-						UE_LOG(LogMutable, Display, TEXT("Id [%s] - NumFiles [%d] -> [%d]."), *ObjectId.ToString(), Category.NumFiles, CategoryOverride->NumFiles);
+						UE_LOGF(LogMutable, Display, "Id [%ls] - NumFiles [%d] -> [%d].", *ObjectId.ToString(), Category.NumFiles, CategoryOverride->NumFiles);
 						Category.NumFiles = CategoryOverride->NumFiles;
 					}
 				}
@@ -2535,7 +2535,7 @@ namespace UE::Mutable::Private
 		{
 			for (auto& CategoryOverride : FileCategoryOverrides)
 			{
-				UE_LOG(LogMutable, Display, TEXT("Id [%s] - Storing Override NumFiles [%d]."), *ObjectId.ToString(), CategoryOverride.NumFiles);
+				UE_LOGF(LogMutable, Display, "Id [%ls] - Storing Override NumFiles [%d].", *ObjectId.ToString(), CategoryOverride.NumFiles);
 			}
 
 			// Upload FileCategory Overrides to DDC.
@@ -2803,7 +2803,7 @@ namespace UE::Mutable::Private
 				break;
 			}
 			default:
-				UE_LOG(LogMutable, Error, TEXT("Unknown DataType found while fixing streaming block files ids."));
+				UE_LOGF(LogMutable, Error, "Unknown DataType found while fixing streaming block files ids.");
 				unimplemented();
 				break;
 			}
@@ -3048,7 +3048,7 @@ private:
 void UCustomizableObjectPrivate::LogMemory()
 {
 	// Log in-memory data for the compiled CustomizableObject
-	UE_LOG(LogMutable, Log, TEXT("CustomizableObject [%s] memory report:"), *GetPublic()->GetName());
+	UE_LOGF(LogMutable, Log, "CustomizableObject [%ls] memory report:", *GetPublic()->GetName());
 	const UE::Mutable::Private::FModel* MuModel = GetModel().Get();
 	if (MuModel)
 	{
@@ -3158,28 +3158,28 @@ void UCustomizableObjectPrivate::LogMemory()
 			ModeleResData += Arch.TotalSize();
 		}
 
-		UE_LOG(LogMutable, Log, TEXT("Total                : %8d"), ModelTotal+StreamableData+ModeleResData);
-		UE_LOG(LogMutable, Log, TEXT("  Streamable         : %8d"), StreamableData);
-		UE_LOG(LogMutable, Log, TEXT("  ModelResources     : %8d"), ModeleResData);
-		UE_LOG(LogMutable, Log, TEXT("  Model              : %8d"), ModelTotal);
-		UE_LOG(LogMutable, Log, TEXT("    ByteCode         : %8d"), ByteCode);
-		UE_LOG(LogMutable, Log, TEXT("    Roms             : %8d"), Roms);
-		UE_LOG(LogMutable, Log, TEXT("    Strings          : %8d"), StringsTotal);
-		UE_LOG(LogMutable, Log, TEXT("    Names			   : %8d"), ConstantNames);
-		UE_LOG(LogMutable, Log, TEXT("    Images           : %8d"), ImageTotal);
-		UE_LOG(LogMutable, Log, TEXT("      Buffer         : %8d"), ConstantImages);
-		UE_LOG(LogMutable, Log, TEXT("      LODIndices     : %8d"), ConstantImageLODIndices);
-		UE_LOG(LogMutable, Log, TEXT("      PermanentBuf   : %8d"), ConstantImageLODPermanent);
-		UE_LOG(LogMutable, Log, TEXT("      Permanent      : %8d"), PermanentImages);
-		UE_LOG(LogMutable, Log, TEXT("    Meshes           : %8d"), MeshesTotal);
-		UE_LOG(LogMutable, Log, TEXT("      Buffer         : %8d"), ConstantMeshes);
-		UE_LOG(LogMutable, Log, TEXT("      ContentIndices : %8d"), ConstantMeshContentIndices);
-		UE_LOG(LogMutable, Log, TEXT("      PermanentBuf   : %8d"), ConstantMeshesPermanent);
-		UE_LOG(LogMutable, Log, TEXT("      Permanent      : %8d"), PermanentMeshes);
-		UE_LOG(LogMutable, Log, TEXT("    Skeletons        : %8d"), SkeletonsTotal);
-		UE_LOG(LogMutable, Log, TEXT("    Sockets	       : %8d"), SocketsTotal);
-		UE_LOG(LogMutable, Log, TEXT("    PhysicBodies     : %8d"), PhysicsTotal);
-		UE_LOG(LogMutable, Log, TEXT("    Parameters       : %8d"), ParametersTotal);
+		UE_LOGF(LogMutable, Log, "Total                : %8d", ModelTotal+StreamableData+ModeleResData);
+		UE_LOGF(LogMutable, Log, "  Streamable         : %8d", StreamableData);
+		UE_LOGF(LogMutable, Log, "  ModelResources     : %8d", ModeleResData);
+		UE_LOGF(LogMutable, Log, "  Model              : %8d", ModelTotal);
+		UE_LOGF(LogMutable, Log, "    ByteCode         : %8d", ByteCode);
+		UE_LOGF(LogMutable, Log, "    Roms             : %8d", Roms);
+		UE_LOGF(LogMutable, Log, "    Strings          : %8d", StringsTotal);
+		UE_LOGF(LogMutable, Log, "    Names			   : %8d", ConstantNames);
+		UE_LOGF(LogMutable, Log, "    Images           : %8d", ImageTotal);
+		UE_LOGF(LogMutable, Log, "      Buffer         : %8d", ConstantImages);
+		UE_LOGF(LogMutable, Log, "      LODIndices     : %8d", ConstantImageLODIndices);
+		UE_LOGF(LogMutable, Log, "      PermanentBuf   : %8d", ConstantImageLODPermanent);
+		UE_LOGF(LogMutable, Log, "      Permanent      : %8d", PermanentImages);
+		UE_LOGF(LogMutable, Log, "    Meshes           : %8d", MeshesTotal);
+		UE_LOGF(LogMutable, Log, "      Buffer         : %8d", ConstantMeshes);
+		UE_LOGF(LogMutable, Log, "      ContentIndices : %8d", ConstantMeshContentIndices);
+		UE_LOGF(LogMutable, Log, "      PermanentBuf   : %8d", ConstantMeshesPermanent);
+		UE_LOGF(LogMutable, Log, "      Permanent      : %8d", PermanentMeshes);
+		UE_LOGF(LogMutable, Log, "    Skeletons        : %8d", SkeletonsTotal);
+		UE_LOGF(LogMutable, Log, "    Sockets	       : %8d", SocketsTotal);
+		UE_LOGF(LogMutable, Log, "    PhysicBodies     : %8d", PhysicsTotal);
+		UE_LOGF(LogMutable, Log, "    Parameters       : %8d", ParametersTotal);
 	}
 
 }
diff --git a/Engine/Plugins/Mutable/Source/CustomizableObject/Private/MuCO/CustomizableObjectDGGUI.cpp b/Engine/Plugins/Mutable/Source/CustomizableObject/Private/MuCO/CustomizableObjectDGGUI.cpp
--- a/Engine/Plugins/Mutable/Source/CustomizableObject/Private/MuCO/CustomizableObjectDGGUI.cpp
+++ b/Engine/Plugins/Mutable/Source/CustomizableObject/Private/MuCO/CustomizableObjectDGGUI.cpp
@@ -38,7 +38,7 @@ void UDGGUI::OpenDGGUI(const int32 SlotID, UCustomizableObjectInstanceUsage* Sel
 		}
 		else
 		{
-			UE_LOG(LogTemp, Warning, TEXT("Could not find the DynamicallyGenerated_DGGUI class inside the specified path. Check the DefaultMutable.ini file."));
+			UE_LOGF(LogTemp, Warning, "Could not find the DynamicallyGenerated_DGGUI class inside the specified path. Check the DefaultMutable.ini file.");
 		}
 	}
 #endif // !UE_BUILD_SHIPPING
diff --git a/Engine/Plugins/Mutable/Source/CustomizableObject/Private/MuCO/CustomizableObjectInstance.cpp b/Engine/Plugins/Mutable/Source/CustomizableObject/Private/MuCO/CustomizableObjectInstance.cpp
--- a/Engine/Plugins/Mutable/Source/CustomizableObject/Private/MuCO/CustomizableObjectInstance.cpp
+++ b/Engine/Plugins/Mutable/Source/CustomizableObject/Private/MuCO/CustomizableObjectInstance.cpp
@@ -880,7 +880,7 @@ bool AreSkeletonsCompatible(const TArray<TObjectPtr<USkeleton>>& InSkeletons)
 						}
 					}
 
-					UE_LOG(LogMutable, Error, TEXT("%s"), *Msg);
+					UE_LOGF(LogMutable, Error, "%ls", *Msg);
 
 #if WITH_EDITOR
 					FNotificationInfo Info(FText::FromString(TEXT("Mutable: Failed to merge skeletons. Invalid parent chain detected. Please check the output log for more information.")));
@@ -897,7 +897,7 @@ bool AreSkeletonsCompatible(const TArray<TObjectPtr<USkeleton>>& InSkeletons)
 				if (ExistingBoneInfo->ParentBoneSkeletonIndex != SkeletonIndex)
 				{
 					// Different skeletons can't be used if they are incompatible with the reference skeleton.
-					UE_LOG(LogMutable, Error, TEXT("[%s] parent bone is different in skeletons [%s] and [%s]."),
+					UE_LOGF(LogMutable, Error, "[%ls] parent bone is different in skeletons [%ls] and [%ls].",
 						*Bone.Name.ToString(),
 						*InSkeletons[SkeletonIndex]->GetName(),
 						*InSkeletons[ExistingBoneInfo->ParentBoneSkeletonIndex]->GetName());
@@ -967,7 +967,7 @@ USkeleton* MergeSkeletons(const TArray<TObjectPtr<USkeleton>>& SkeletonsToMerge,
 			Msg += FString::Printf(TEXT(" [%s]"), *Params.SkeletonsToMerge[SkeletonIndex]->GetName());
 		}
 		
-		UE_LOG(LogMutable, Error, TEXT("%s"), *Msg);
+		UE_LOGF(LogMutable, Error, "%ls", *Msg);
 	}
 	else
 	{
@@ -1670,7 +1670,7 @@ bool DoComponentsNeedUpdate(const TSharedRef<FUpdateContextPrivate>& Context, bo
 
 				if (!LOD->Mesh->IsReference() && LOD->Mesh->GetVertexCount() == 0)
 				{
-					UE_LOG(LogMutable, Error, TEXT("Failed to generate SkeletalMesh for CO Instance [%s]. CO [%s] has invalid geometry for LOD [%d] Component [%d]."),
+					UE_LOGF(LogMutable, Error, "Failed to generate SkeletalMesh for CO Instance [%ls]. CO [%ls] has invalid geometry for LOD [%d] Component [%d].",
 						*Instance->GetName(), *CustomizableObject->GetName(),
 						LODIndex, ComponentIndex);
 					bHasInvalidMesh = true; // Unreal does not support empty LODs.
@@ -1709,7 +1709,7 @@ bool UCustomizableInstancePrivate::UpdateSkeletalMesh_PostBeginUpdate0(UCustomiz
 
 	if (!CustomizableObject)
 	{
-		UE_LOG(LogMutable, Warning, TEXT("Failed to generate SkeletalMesh for CO Instance %s. It does not have a CO."), *Instance->GetName());
+		UE_LOGF(LogMutable, Warning, "Failed to generate SkeletalMesh for CO Instance %ls. It does not have a CO.", *Instance->GetName());
 		Context->UpdateResult = EUpdateResult::Error;
 		return false;
 	}
@@ -1717,7 +1717,7 @@ bool UCustomizableInstancePrivate::UpdateSkeletalMesh_PostBeginUpdate0(UCustomiz
 	// We can not handle empty meshes, clear any generated mesh and return
 	if (bHasInvalidMesh)
 	{
-		UE_LOG(LogMutable, Warning, TEXT("Failed to generate SkeletalMesh for CO Instance %s. CO [%s]"), *Instance->GetName(), *GetNameSafe(CustomizableObject));
+		UE_LOGF(LogMutable, Warning, "Failed to generate SkeletalMesh for CO Instance %ls. CO [%ls]", *Instance->GetName(), *GetNameSafe(CustomizableObject));
 		Context->UpdateResult = EUpdateResult::Error;
 		return false;
 	}
@@ -2788,7 +2788,7 @@ TUniquePtr<FTexturePlatformData> MutableCreateImagePlatformData(UE::Mutable::Pri
 
 	if (SizeX <= 0 || SizeY <= 0)
 	{
-		UE_LOG(LogMutable, Warning, TEXT("Invalid parameters specified for UCustomizableInstancePrivate::MutableCreateImagePlatformData()"));
+		UE_LOGF(LogMutable, Warning, "Invalid parameters specified for UCustomizableInstancePrivate::MutableCreateImagePlatformData()");
 		return nullptr;
 	}
 
@@ -2868,7 +2868,7 @@ TUniquePtr<FTexturePlatformData> MutableCreateImagePlatformData(UE::Mutable::Pri
 	//check(SizeX == MaxPossibleSize || SizeY == MaxPossibleSize || FullLODCount == 1);
 	if (!(SizeX == MaxPossibleSize || SizeY == MaxPossibleSize || FullLODCount == 1))
 	{
-		UE_LOG(LogMutable, Warning, TEXT("Building instance: unsupported texture size %d x %d."), SizeX, SizeY);
+		UE_LOGF(LogMutable, Warning, "Building instance: unsupported texture size %d x %d.", SizeX, SizeY);
 		//return nullptr;
 	}
 
@@ -2880,13 +2880,13 @@ TUniquePtr<FTexturePlatformData> MutableCreateImagePlatformData(UE::Mutable::Pri
 	case UE::Mutable::Private::EImageFormat::RGB_UByte:
 		// performance penalty. can happen in states that remove compression.
 		PlatformFormat = PF_R8G8B8A8;	
-		UE_LOG(LogMutable, Display, TEXT("Building instance: a texture was generated in a format not supported by the hardware (RGB), this results in an additional conversion, so a performance penalty."));
+		UE_LOGF(LogMutable, Display, "Building instance: a texture was generated in a format not supported by the hardware (RGB), this results in an additional conversion, so a performance penalty.");
 		break; 
 
 	case UE::Mutable::Private::EImageFormat::BGRA_UByte:			
 		// performance penalty. can happen with texture parameter images.
 		PlatformFormat = PF_R8G8B8A8;	
-		UE_LOG(LogMutable, Display, TEXT("Building instance: a texture was generated in a format not supported by the hardware (BGRA), this results in an additional conversion, so a performance penalty."));
+		UE_LOGF(LogMutable, Display, "Building instance: a texture was generated in a format not supported by the hardware (BGRA), this results in an additional conversion, so a performance penalty.");
 		break;
 
 	// Good cases:
@@ -2902,7 +2902,7 @@ TUniquePtr<FTexturePlatformData> MutableCreateImagePlatformData(UE::Mutable::Pri
 	case UE::Mutable::Private::EImageFormat::ASTC_4x4_RG_LDR:	PlatformFormat = PF_ASTC_4x4;	break;
 	default:
 		// Cannot prepare texture if it's not in the right format, this can happen if mutable is in debug mode or in case of bugs
-		UE_LOG(LogMutable, Warning, TEXT("Building instance: a texture was generated in an unsupported format, it will be converted to Unreal with a performance penalty."));
+		UE_LOGF(LogMutable, Warning, "Building instance: a texture was generated in an unsupported format, it will be converted to Unreal with a performance penalty.");
 
 		switch (UE::Mutable::Private::GetImageFormatData(MutableFormat).Channels)
 		{
@@ -3077,7 +3077,7 @@ void ConvertImage(UTexture2D* Texture, UE::Mutable::Private::TManagedPtr<const U
 	// Hack: This format is unsupported in UE, but it shouldn't happen in production.
 	if (MutableFormat == UE::Mutable::Private::EImageFormat::RGB_UByte)
 	{
-		UE_LOG(LogMutable, Warning, TEXT("Building instance: a texture was generated in RGB format, which is slow to convert to Unreal."));
+		UE_LOGF(LogMutable, Warning, "Building instance: a texture was generated in RGB format, which is slow to convert to Unreal.");
 
 		// Expand the image.
 		UE::Mutable::Private::TManagedPtr<UE::Mutable::Private::FImage> Converted = 
@@ -3101,7 +3101,7 @@ void ConvertImage(UTexture2D* Texture, UE::Mutable::Private::TManagedPtr<const U
 	}
 	else if (MutableFormat == UE::Mutable::Private::EImageFormat::BGRA_UByte)
 	{
-		UE_LOG(LogMutable, Warning, TEXT("Building instance: a texture was generated in BGRA format, which is slow to convert to Unreal."));
+		UE_LOGF(LogMutable, Warning, "Building instance: a texture was generated in BGRA format, which is slow to convert to Unreal.");
 
 		MUTABLE_CPUPROFILER_SCOPE(Swizzle);
 		// Swizzle the image.
@@ -3877,7 +3877,7 @@ bool BuildOrCopyRenderData(const TSharedRef<FUpdateContextPrivate>& Context, USk
 		// There could be components without a mesh in LODs
 		if (!LOD->Mesh || LOD->SurfaceCount == 0)
 		{
-			UE_LOG(LogMutable, Warning, TEXT("Building instance: generated mesh [%s] has LOD [%d] of object component index [%d] with no mesh.")
+			UE_LOGF(LogMutable, Warning, "Building instance: generated mesh [%ls] has LOD [%d] of object component index [%d] with no mesh."
 				, *SkeletalMesh->GetName()
 				, LODIndex
 				, Component->Id);
@@ -4062,7 +4062,7 @@ bool BuildOrCopyRenderData(const TSharedRef<FUpdateContextPrivate>& Context, USk
 
 			const FString Msg = FString::Printf(TEXT("Customizable Object [%s] requires of Skinning - 'Support 16 Bit Bone Index' to be enabled. Please, update the Project Settings."),
 				*CustomizableObject->GetName());
-			UE_LOG(LogMutable, Error, TEXT("%s"), *Msg);
+			UE_LOGF(LogMutable, Error, "%ls", *Msg);
 
 #if WITH_EDITOR
 			FNotificationInfo Info(FText::FromString(Msg));
@@ -4358,7 +4358,7 @@ UTexture* BuildTexture(const TSharedRef<FUpdateContextPrivate>& Context, const i
 
 		if (!Context->ImageToPlatformDataMap.Contains(Image->ImageID))
 		{
-			UE_LOG(LogMutable, Error, TEXT("Required image [%s] with ID [%s] was not generated in the mutable thread, and it is not cached."),
+			UE_LOGF(LogMutable, Error, "Required image [%ls] with ID [%ls] was not generated in the mutable thread, and it is not cached.",
 			*Props.TextureParameterName,
 			*Image->ImageID.ToString());
 			
@@ -4979,7 +4979,7 @@ void UCustomizableObjectInstance::Bake(const FBakingConfiguration& InBakingConfi
 	else
 	{
 		// Notify of the error
-		UE_LOG(LogMutable, Error, TEXT("The module \" ICustomizableObjectEditorModule \" could not be loaded and therefore the baking operation could not be started."));
+		UE_LOGF(LogMutable, Error, "The module \" ICustomizableObjectEditorModule \" could not be loaded and therefore the baking operation could not be started.");
 		if (InBakingConfiguration.OnBakeOperationCompletedCallback.IsBound())
 		{
 			FCustomizableObjectInstanceBakeOutput Output;
@@ -5090,7 +5090,7 @@ namespace UE::Mutable::Private
 		if (!Delegate.IsBound())
 		{
 			FString ErrorMsg = FString::Printf(TEXT("Attempting to iterate over AnimInstances with an unbound delegate for component [%s]."), *ComponentName.ToString());
-			UE_LOG(LogMutable, Warning, TEXT("%s"), *ErrorMsg);
+			UE_LOGF(LogMutable, Warning, "%ls", *ErrorMsg);
 #if WITH_EDITOR
 			FMessageLogModule& MessageLogModule = FModuleManager::LoadModuleChecked<FMessageLogModule>("MessageLog");
 			MessageLogModule.RegisterLogListing(FName("Mutable"), FText::FromString(FString("Mutable")));
@@ -5218,7 +5218,7 @@ void UCustomizableObjectInstance::AnimInstanceFixup(UAnimInstance* InAnimInstanc
 #if WITH_EDITOR
 		if (bPropertyMismatchFound)
 		{
-			UE_LOG(LogMutable, Warning, TEXT("AnimBp %s is not in sync with the data stored in the CO %s. A CO recompilation may be needed."),
+			UE_LOGF(LogMutable, Warning, "AnimBp %ls is not in sync with the data stored in the CO %ls. A CO recompilation may be needed.",
 				*AnimInstanceClass.Get()->GetName(), 
 				*GetCustomizableObject()->GetName());
 		}
diff --git a/Engine/Plugins/Mutable/Source/CustomizableObject/Private/MuCO/CustomizableObjectInstanceDescriptor.cpp b/Engine/Plugins/Mutable/Source/CustomizableObject/Private/MuCO/CustomizableObjectInstanceDescriptor.cpp
--- a/Engine/Plugins/Mutable/Source/CustomizableObject/Private/MuCO/CustomizableObjectInstanceDescriptor.cpp
+++ b/Engine/Plugins/Mutable/Source/CustomizableObject/Private/MuCO/CustomizableObjectInstanceDescriptor.cpp
@@ -310,7 +310,7 @@ void FCustomizableObjectInstanceDescriptor::SaveDescriptor(FArchive& Ar, bool bU
 			
 		case EMutableParameterType::None:
 		{
-			UE_LOG(LogMutable, Warning, TEXT("Found unknown parameter [%s] when saving a Customizable Object Instance Descriptor"), *Name);
+			UE_LOGF(LogMutable, Warning, "Found unknown parameter [%ls] when saving a Customizable Object Instance Descriptor", *Name);
 			break;
 		}
  
@@ -565,7 +565,7 @@ void FCustomizableObjectInstanceDescriptor::LoadDescriptor(FArchive& Ar)
 			
 		case EMutableParameterType::None:
 		{
-			UE_LOG(LogMutable, Warning, TEXT("Found unknown parameter [%s] when loading a Customizable Object Instance Descriptor"), *Name);
+			UE_LOGF(LogMutable, Warning, "Found unknown parameter [%ls] when loading a Customizable Object Instance Descriptor", *Name);
 			break;
 		}
  
@@ -976,7 +976,7 @@ void FCustomizableObjectInstanceDescriptor::ReloadParameters()
 	
 	if (!CustomizableObject->GetPrivate()->GetModel())
 	{
-		UE_LOG(LogMutable, Warning, TEXT("[ReloadParametersFromObject] No model in object [%s], generated empty parameters for [%s] "), *CustomizableObject->GetName(), *CustomizableObject->GetName());
+		UE_LOGF(LogMutable, Warning, "[ReloadParametersFromObject] No model in object [%ls], generated empty parameters for [%ls] ", *CustomizableObject->GetName(), *CustomizableObject->GetName());
 		return;
 	}
 
@@ -984,7 +984,7 @@ void FCustomizableObjectInstanceDescriptor::ReloadParameters()
 	
 	if(!CustomizableObject->GetPrivate()->GetModelResources())
 	{
-		UE_LOG(LogMutable, Warning, TEXT("[ReloadParametersFromObject] No model resources in object [%s], generated empty parameters for [%s] "), *CustomizableObject->GetName(), *CustomizableObject->GetName());
+		UE_LOGF(LogMutable, Warning, "[ReloadParametersFromObject] No model resources in object [%ls], generated empty parameters for [%ls] ", *CustomizableObject->GetName(), *CustomizableObject->GetName());
 		return;
 	}
 	
@@ -1661,7 +1661,7 @@ void FCustomizableObjectInstanceDescriptor::SetIntParameterSelectedOption(const
 			CustomizableObject->GetPrivate()->FindIntParameterValue(ParameterIndexInObject, SelectedOption),
 			*GetAvailableOptionsString(*CustomizableObject, ParameterIndexInObject)
 		);
-		UE_LOG(LogMutable, Error, TEXT("%s"), *Message);
+		UE_LOGF(LogMutable, Error, "%ls", *Message);
 #endif
 		return;
 	}
@@ -2077,7 +2077,7 @@ void FCustomizableObjectInstanceDescriptor::SetExternalTypeParameterSelectedOpti
 	const UModelResources& ModelResources = CustomizableObject->GetPrivate()->GetModelResourcesChecked();
 	if (!ExternalTypeValue.GetScriptStruct()->IsChildOf(ModelResources.ExternalTypeParameterDefaultValues[ExternalTypeParamName].GetScriptStruct()))
 	{
-		UE_LOG(LogMutable, Error, TEXT("External Parameter %s incorrect type."), *ExternalTypeParamName);
+		UE_LOGF(LogMutable, Error, "External Parameter %ls incorrect type.", *ExternalTypeParamName);
 		return;
 	}
 	
@@ -2264,7 +2264,7 @@ void FCustomizableObjectInstanceDescriptor::SetProjectorParameterSelectedOption(
 	FCustomizableObjectProjectorParameterValue& ProjectorParameter = ProjectorParameters[ParameterIndexInInstance];
 	if (ProjectorValue.ProjectionType != ProjectorParameter.Value.ProjectionType)
 	{
-		UE_LOG(LogMutable, Warning, TEXT("Changing the Projection type is not allowed once the CO has been compiled. In order to change it update the CO Projector Parameter settings and recompile the CO."));
+		UE_LOGF(LogMutable, Warning, "Changing the Projection type is not allowed once the CO has been compiled. In order to change it update the CO Projector Parameter settings and recompile the CO.");
 	}
  
 	// New Value
@@ -2740,7 +2740,7 @@ void FCustomizableObjectInstanceDescriptor::SetCurrentState(const FString& State
 	}
 	else
 	{
-		UE_LOG(LogMutable, Error, TEXT("%hs: Unknown %s state."), __FUNCTION__, *StateName);
+		UE_LOGF(LogMutable, Error, "%hs: Unknown %ls state.", __FUNCTION__, *StateName);
 	}
 }
  
diff --git a/Engine/Plugins/Mutable/Source/CustomizableObject/Private/MuCO/CustomizableObjectInstanceUsage.cpp b/Engine/Plugins/Mutable/Source/CustomizableObject/Private/MuCO/CustomizableObjectInstanceUsage.cpp
--- a/Engine/Plugins/Mutable/Source/CustomizableObject/Private/MuCO/CustomizableObjectInstanceUsage.cpp
+++ b/Engine/Plugins/Mutable/Source/CustomizableObject/Private/MuCO/CustomizableObjectInstanceUsage.cpp
@@ -33,7 +33,7 @@ void UCustomizableObjectInstanceUsagePrivate::Callbacks()
 
 		if (GetPublic()->UpdatedDelegate.IsBound() && CustomizableSkeletalComponent->UpdatedDelegate.IsBound())
 		{
-			UE_LOG(LogMutable, Error, TEXT("The UpdatedDelegate is bound both in the UCustomizableObjectInstanceUsage and in its parent CustomizableSkeletalComponent. Only one should be bound."));
+			UE_LOGF(LogMutable, Error, "The UpdatedDelegate is bound both in the UCustomizableObjectInstanceUsage and in its parent CustomizableSkeletalComponent. Only one should be bound.");
 			ensure(false);
 		}
 	}
@@ -287,4 +287,4 @@ UCustomizableObjectInstanceUsagePrivate* UCustomizableObjectInstanceUsage::GetPr
 const UCustomizableObjectInstanceUsagePrivate* UCustomizableObjectInstanceUsage::GetPrivate() const
 {
 	return Private;
-}
\ No newline at end of file
+}
diff --git a/Engine/Plugins/Mutable/Source/CustomizableObject/Private/MuCO/CustomizableObjectMipDataProvider.cpp b/Engine/Plugins/Mutable/Source/CustomizableObject/Private/MuCO/CustomizableObjectMipDataProvider.cpp
--- a/Engine/Plugins/Mutable/Source/CustomizableObject/Private/MuCO/CustomizableObjectMipDataProvider.cpp
+++ b/Engine/Plugins/Mutable/Source/CustomizableObject/Private/MuCO/CustomizableObjectMipDataProvider.cpp
@@ -50,7 +50,7 @@ FMutableTextureMipDataProvider::FMutableTextureMipDataProvider(const UTexture* T
 
 void FMutableTextureMipDataProvider::PrintWarningAndAdvanceToCleanup()
 {
-	UE_LOG(LogMutable, Warning, TEXT("Tried to update a mip from a Customizable Object being compiled, cancelling mip update."));
+	UE_LOGF(LogMutable, Warning, "Tried to update a mip from a Customizable Object being compiled, cancelling mip update.");
 	AdvanceTo(ETickState::CleanUp, ETickThread::Async);
 }
 
@@ -175,7 +175,7 @@ namespace Impl
 					{
 						MUTABLE_CPUPROFILER_SCOPE(GetImage_MipFix);
 
-						UE_LOG(LogMutable, Warning, TEXT("Mutable generated an incomplete mip chain for image."));
+						UE_LOGF(LogMutable, Warning, "Mutable generated an incomplete mip chain for image.");
 
 						// Force the right number of mips. The missing data will be black.
 						UE::Mutable::Private::TManagedPtr<UE::Mutable::Private::FImage> NewImage = 
@@ -522,14 +522,14 @@ bool FMutableTextureMipDataProvider::PollMips(const FTextureUpdateSyncOptions& S
 				}
 				else
 				{
-					UE_LOG(LogMutable, Warning, TEXT("Mip data has incorrect size."));
+					UE_LOGF(LogMutable, Warning, "Mip data has incorrect size.");
 					FMemory::Memzero(Dest, Level.DataSize);
 				}
 			}
 			else
 			{
 				// Mutable didn't generate all the expected mips
-				UE_LOG(LogMutable, Warning, TEXT("Mutable image is missing mips."));
+				UE_LOGF(LogMutable, Warning, "Mutable image is missing mips.");
 				FMemory::Memzero(Dest, Level.DataSize);
 			}
 			++MipIndex;
diff --git a/Engine/Plugins/Mutable/Source/CustomizableObject/Private/MuCO/CustomizableObjectModule.cpp b/Engine/Plugins/Mutable/Source/CustomizableObject/Private/MuCO/CustomizableObjectModule.cpp
--- a/Engine/Plugins/Mutable/Source/CustomizableObject/Private/MuCO/CustomizableObjectModule.cpp
+++ b/Engine/Plugins/Mutable/Source/CustomizableObject/Private/MuCO/CustomizableObjectModule.cpp
@@ -226,7 +226,7 @@ void FCustomizableObjectModule::RegisterExternalOperation(const UScriptStruct* O
 
 	if (!Operation->IsChildOf(UE::Mutable::FExternalOperation::StaticStruct()))
 	{
-		UE_LOG(LogMutable, Error, TEXT("Unable to register operation. Operation that does not inherit from FMutableOperation: %s"), *Operation->GetName());
+		UE_LOGF(LogMutable, Error, "Unable to register operation. Operation that does not inherit from FMutableOperation: %ls", *Operation->GetName());
 		return;
 	}
 
@@ -433,37 +433,37 @@ void PrintParticipatingPackagesDiff(const TArray<FName>& OutOfDatePackages, cons
 
 	if (bReleaseVersion)
 	{
-		UE_LOG(LogMutable, Display, TEXT("Release Version changed."));
+		UE_LOGF(LogMutable, Display, "Release Version changed.");
 	}
 	
 	if (OutOfDatePackages.Num())
 	{
-		UE_LOG(LogMutable, Display, TEXT("Packages out of date (%i):"), OutOfDatePackages.Num());
+		UE_LOGF(LogMutable, Display, "Packages out of date (%i):", OutOfDatePackages.Num());
 	}
 				
 	for (int32 Index = 0; Index < FMath::Min(MaxLogLines, OutOfDatePackages.Num()); ++Index)
 	{
-		UE_LOG(LogMutable, Display, TEXT("%s"), *OutOfDatePackages[Index].ToString());
+		UE_LOGF(LogMutable, Display, "%ls", *OutOfDatePackages[Index].ToString());
 	}
 
 	if (AddedPackages.Num())
 	{
-		UE_LOG(LogMutable, Display, TEXT("Added packages (%i):"), AddedPackages.Num());
+		UE_LOGF(LogMutable, Display, "Added packages (%i):", AddedPackages.Num());
 	}
 				
 	for (int32 Index = 0; Index < FMath::Min(MaxLogLines, AddedPackages.Num()); ++Index)
 	{
-		UE_LOG(LogMutable, Display, TEXT("%s"), *AddedPackages[Index].ToString());
+		UE_LOGF(LogMutable, Display, "%ls", *AddedPackages[Index].ToString());
 	}
 
 	if (RemovedPackages.Num())
 	{
-		UE_LOG(LogMutable, Display, TEXT("Removed packages (%i):"), RemovedPackages.Num());
+		UE_LOGF(LogMutable, Display, "Removed packages (%i):", RemovedPackages.Num());
 	}
 				
 	for (int32 Index = 0; Index < FMath::Min(MaxLogLines, RemovedPackages.Num()); ++Index)
 	{
-		UE_LOG(LogMutable, Display, TEXT("%s"), *RemovedPackages[Index].ToString());
+		UE_LOGF(LogMutable, Display, "%ls", *RemovedPackages[Index].ToString());
 	}
 }
 
diff --git a/Engine/Plugins/Mutable/Source/CustomizableObject/Private/MuCO/CustomizableObjectSystem.cpp b/Engine/Plugins/Mutable/Source/CustomizableObject/Private/MuCO/CustomizableObjectSystem.cpp
--- a/Engine/Plugins/Mutable/Source/CustomizableObject/Private/MuCO/CustomizableObjectSystem.cpp
+++ b/Engine/Plugins/Mutable/Source/CustomizableObject/Private/MuCO/CustomizableObjectSystem.cpp
@@ -354,7 +354,7 @@ void OutOfDateWarning_Async(const TSharedRef<FOutOfDateWarningContext>& Context)
 			
 				if (!GEngine->OnScreenDebugMessageExists(Key))
 				{
-					UE_LOG(LogMutable, Display, TEXT("Customizable Object [%s] compilation out of date. Changes since last compilation:"), *Object->GetName());
+					UE_LOGF(LogMutable, Display, "Customizable Object [%ls] compilation out of date. Changes since last compilation:", *Object->GetName());
 					PrintParticipatingPackagesDiff(OutOfDatePackages, AddedPackages, RemovedPackages, bVersionDiff);
 				}
 			
@@ -545,7 +545,7 @@ UCustomizableObjectSystem* UCustomizableObjectSystem::GetInstance()
 {
 	if (!UCustomizableObjectSystemPrivate::SSystem)
 	{
-		UE_LOG(LogMutable, Log, TEXT("Creating Mutable Customizable Object System."));
+		UE_LOGF(LogMutable, Log, "Creating Mutable Customizable Object System.");
 
 		check(IsInGameThread());
 
@@ -1170,7 +1170,7 @@ void UpdateSkeletalMesh(const TSharedRef<FUpdateContextPrivate>& Context)
 
 	if (ParentsWithOverrideMaterialsUpdatedCount > 0)
 	{
-		UE_LOG(LogMutable, Verbose, TEXT("A total of %u Skeletal Mesh Components got their override materials updated."), ParentsWithOverrideMaterialsUpdatedCount);
+		UE_LOGF(LogMutable, Verbose, "A total of %u Skeletal Mesh Components got their override materials updated.", ParentsWithOverrideMaterialsUpdatedCount);
 	}
 }
 
@@ -1191,7 +1191,7 @@ void UCustomizableObjectSystemPrivate::EnqueueUpdateSkeletalMesh(const TSharedRe
 
 	const FName ObjectName = Object ? Object->GetFName() : NAME_None;
 	const FName InstanceName = Instance ? Instance->GetFName() : NAME_None;
-	UE_LOG(LogMutable, Verbose, TEXT("Enqueue Update Skeletal Mesh Async. CustomizableObject=%s Instance=%s, Frame=%d"), *ObjectName.ToString(), *InstanceName.ToString(), GFrameNumber);
+	UE_LOGF(LogMutable, Verbose, "Enqueue Update Skeletal Mesh Async. CustomizableObject=%ls Instance=%ls, Frame=%d", *ObjectName.ToString(), *InstanceName.ToString(), GFrameNumber);
 	
 	if (!Context->IsContextValid()) 
 	{
@@ -1297,7 +1297,7 @@ void UCustomizableObjectSystemPrivate::EnqueueUpdateSkeletalMesh(const TSharedRe
 		String += FString::Printf(TEXT("=== DESCRIPTOR ===\n%s"), *Instance->GetPrivate()->GetDescriptor().ToString());
 		String += "================================";
 
-		UE_LOG(LogMutable, Log, TEXT("%s"), *String);
+		UE_LOGF(LogMutable, Log, "%ls", *String);
 	}
 
 	Context->StartQueueTime = FPlatformTime::Seconds();
@@ -1389,7 +1389,7 @@ bool UCustomizableObjectSystem::LockObject(UCustomizableObject* InObject)
 		if (Private->CurrentInstanceBeingUpdated &&
 			Private->CurrentInstanceBeingUpdated->GetCustomizableObject() == InObject)
 		{
-			UE_LOG(LogMutable, Warning, TEXT("---- failed to lock object %s"), *InObject->GetName());
+			UE_LOGF(LogMutable, Warning, "---- failed to lock object %ls", *InObject->GetName());
 
 			return false;
 		}
@@ -1403,7 +1403,7 @@ bool UCustomizableObjectSystem::LockObject(UCustomizableObject* InObject)
 		// the compile button but the compilation quits anyway because there are pending operations
 		if (CheckIfDiskOrMipUpdateOperationsPending(*InObject))
 		{
-			UE_LOG(LogMutable, Warning, TEXT("%s"), *Message);
+			UE_LOGF(LogMutable, Warning, "%ls", *Message);
 
 			return false;
 		}
@@ -1420,7 +1420,7 @@ bool UCustomizableObjectSystem::LockObject(UCustomizableObject* InObject)
 		// But some could have started between the first CheckIfDiskOrMipUpdateOperationsPending and the lock a few lines back, so check again
 		if (CheckIfDiskOrMipUpdateOperationsPending(*InObject))
 		{
-			UE_LOG(LogMutable, Warning, TEXT("%s"), *Message);
+			UE_LOGF(LogMutable, Warning, "%ls", *Message);
 
 			// Unlock and return because the pending operations cannot be easily stopped now, the compilation hasn't started and the CO
 			// hasn't changed state yet. It's simpler to quit the compilation, unlock and let the user try to compile again
@@ -1449,7 +1449,7 @@ bool UCustomizableObjectSystem::LockObject(UCustomizableObject* InObject)
 	else
 	{
 		FString ObjectName = InObject ? InObject->GetName() : FString("null");
-		UE_LOG(LogMutable, Warning, TEXT("Failed to lock the object [%s] because it was null or the system was null or partially destroyed."), *ObjectName);
+		UE_LOGF(LogMutable, Warning, "Failed to lock the object [%ls] because it was null or the system was null or partially destroyed.", *ObjectName);
 
 		return false;
 	}
@@ -2133,7 +2133,7 @@ namespace impl
 					TStrongObjectPtr<const UTexture> CachedTexture = CachedWeakTexture.TryPin(bWasPinSuccessful);
 					if (CachedTexture.IsValid())
 					{
-						UE_LOG(LogMutable, VeryVerbose, TEXT("Texture resource with id [%s] is cached."), *Image->ImageID.ToString());
+						UE_LOGF(LogMutable, VeryVerbose, "Texture resource with id [%ls] is cached.", *Image->ImageID.ToString());
 						OperationData->Objects.Add(CachedTexture);
 						continue;
 					}
@@ -2176,7 +2176,7 @@ namespace impl
 					if (bSizeMissmatch)
 					{
 						// Generate a correctly-sized but empty image instead, to avoid crashes.
-						UE_LOG(LogMutable, Warning, TEXT("Mutable generated a wrongly-sized image %s."), *Image->ImageID.ToString());
+						UE_LOGF(LogMutable, Warning, "Mutable generated a wrongly-sized image %ls.", *Image->ImageID.ToString());
 						Image->Image = UE::Mutable::Private::MakeManaged<UE::Mutable::Private::FImage>(MipSizeX, MipSizeY, FullLODCount - MipsToSkip, Image->Image->GetFormat(), UE::Mutable::Private::EInitializationType::Black);
 					}
 
@@ -2192,7 +2192,7 @@ namespace impl
 					{
 						MUTABLE_CPUPROFILER_SCOPE(GetImage_MipFix);
 
-						UE_LOG(LogMutable, Warning, TEXT("Mutable generated an incomplete mip chain for image %s."), *Image->ImageID.ToString());
+						UE_LOGF(LogMutable, Warning, "Mutable generated an incomplete mip chain for image %ls.", *Image->ImageID.ToString());
 
 						// Force the right number of mips. The missing data will be black.
 						const UE::Mutable::Private::TManagedPtr<UE::Mutable::Private::FImage> NewImage = UE::Mutable::Private::MakeManaged<UE::Mutable::Private::FImage>(Image->Image->GetSizeX(), Image->Image->GetSizeY(), FullMipCount, Image->Image->GetFormat(), UE::Mutable::Private::EInitializationType::Black);
@@ -2380,7 +2380,7 @@ namespace impl
 
 					if (ComponentNames.Contains(ComponentName)) // Components with the same name. Abort.
 					{
-						UE_LOG(LogMutable, Error, TEXT("Update error. Trying to generate multiple components with the same name: %s"),
+						UE_LOGF(LogMutable, Error, "Update error. Trying to generate multiple components with the same name: %ls",
 							*ComponentName.ToString());
 						
 						OperationData->UpdateResult = EUpdateResult::Error;
@@ -3362,7 +3362,7 @@ void UCustomizableObjectSystem::AddUncompiledCOWarning(const UCustomizableObject
 #endif
 
 	// Also log an error so if this happens as part of a bug report we'll have this info.
-	UE_LOG(LogMutable, Error, TEXT("%s"), *ErrorString);
+	UE_LOGF(LogMutable, Error, "%ls", *ErrorString);
 }
 
 
@@ -3389,7 +3389,7 @@ bool UCustomizableObjectSystem::ShouldClearWorkingMemoryOnUpdateEnd()
 void UCustomizableObjectSystem::SetWorkingMemory(int32 KBytes)
 {
 	CVarWorkingMemoryKB->Set( KBytes );
-	UE_LOG(LogMutable, Log, TEXT("Working Memory set to %i kilobytes."), KBytes);
+	UE_LOGF(LogMutable, Log, "Working Memory set to %i kilobytes.", KBytes);
 }
 
 
diff --git a/Engine/Plugins/Mutable/Source/CustomizableObject/Private/MuCO/FMutableTaskGraph.cpp b/Engine/Plugins/Mutable/Source/CustomizableObject/Private/MuCO/FMutableTaskGraph.cpp
--- a/Engine/Plugins/Mutable/Source/CustomizableObject/Private/MuCO/FMutableTaskGraph.cpp
+++ b/Engine/Plugins/Mutable/Source/CustomizableObject/Private/MuCO/FMutableTaskGraph.cpp
@@ -293,7 +293,7 @@ void FMutableTaskGraph::TryLaunchMutableTaskLowPriority(bool bFromMutableTask)
 
 	if (bTimeLimit)
 	{
-		UE_LOG(LogMutable, Verbose, TEXT("Low Priority Mutable Task launched due to time limit (%f)! Waited for: %f"), TimeLimit, TimeElapsed)
+		UE_LOGF(LogMutable, Verbose, "Low Priority Mutable Task launched due to time limit (%f)! Waited for: %f", TimeLimit, TimeElapsed)
 	}
 }
 
diff --git a/Engine/Plugins/Mutable/Source/CustomizableObject/Private/MuCO/LogBenchmarkUtil.cpp b/Engine/Plugins/Mutable/Source/CustomizableObject/Private/MuCO/LogBenchmarkUtil.cpp
--- a/Engine/Plugins/Mutable/Source/CustomizableObject/Private/MuCO/LogBenchmarkUtil.cpp
+++ b/Engine/Plugins/Mutable/Source/CustomizableObject/Private/MuCO/LogBenchmarkUtil.cpp
@@ -36,7 +36,7 @@ namespace LogBenchmarkUtil
 	void Write(FArchive& Archive, const FString& Text )
 	{
 		const FString LogUID = TEXT("MUTABLE_BENCHMARK");
-		UE_LOG(LogMutable, Display, TEXT("%s : %s"),*LogUID, *Text);
+		UE_LOGF(LogMutable, Display, "%ls : %ls",*LogUID, *Text);
 		
 		const FString ComposedString =  FString::Printf(TEXT("%s\n"),*Text);
 		const FStringView StringView = ComposedString;
diff --git a/Engine/Plugins/Mutable/Source/CustomizableObject/Private/MuCO/UnrealConversionUtils.cpp b/Engine/Plugins/Mutable/Source/CustomizableObject/Private/MuCO/UnrealConversionUtils.cpp
--- a/Engine/Plugins/Mutable/Source/CustomizableObject/Private/MuCO/UnrealConversionUtils.cpp
+++ b/Engine/Plugins/Mutable/Source/CustomizableObject/Private/MuCO/UnrealConversionUtils.cpp
@@ -561,13 +561,13 @@ namespace UnrealConversionUtils
 		if (MutableIndexCount == 0)
 		{
 			// Copy indices from an empty buffer
-			UE_LOG(LogMutable, Error, TEXT("UCustomizableInstancePrivateData::BuildSkeletalMeshRenderData is converting an empty mesh."));
+			UE_LOGF(LogMutable, Error, "UCustomizableInstancePrivateData::BuildSkeletalMeshRenderData is converting an empty mesh.");
 			return false;
 		}
 
 		if (InMutableMesh->GetIndexBuffers().IsDescriptor())
 		{
-			UE_LOG(LogMutable, Error, TEXT("UCustomizableInstancePrivateData::BuildSkeletalMeshRenderData is converting a mesh descriptor."));
+			UE_LOGF(LogMutable, Error, "UCustomizableInstancePrivateData::BuildSkeletalMeshRenderData is converting a mesh descriptor.");
 			return false;
 		}
 
@@ -1506,7 +1506,7 @@ namespace UnrealConversionUtils
 				{
 					// This may happen if we are requesting a LOD that is not present in the provided mesh.
 					// TODO: an empty mesh is returned, but should we use the last LOD instead?
-					UE_LOG(LogMutable, Warning, TEXT("Trying to convert mesh [%s] LOD %d but it only has %d LODs"), 
+					UE_LOGF(LogMutable, Warning, "Trying to convert mesh [%ls] LOD %d but it only has %d LODs", 
 						*SkeletalMesh->GetName(), LODIndex, Data->LODRenderData.Num() );
 					
 					return;
@@ -1598,11 +1598,11 @@ namespace UnrealConversionUtils
 				{
 					if (Context->OriginalLOD->bStreamedDataInlined && !SkeletalMesh->NeedCPUData(LODIndex))
 					{
-						UE_LOG(LogMutable, Error, TEXT("Unable to read LOD %d Data from the CPU. Enable bAllowCPUAccess on the Skeletal Mesh %s."), LODIndex, *SkeletalMesh->GetName());
+						UE_LOGF(LogMutable, Error, "Unable to read LOD %d Data from the CPU. Enable bAllowCPUAccess on the Skeletal Mesh %ls.", LODIndex, *SkeletalMesh->GetName());
 					}
 					else
 					{
-						UE_LOG(LogMutable, Error, TEXT("Unable to read LOD %d Data from Skeletal Mesh %s."), LODIndex, *SkeletalMesh->GetName());
+						UE_LOGF(LogMutable, Error, "Unable to read LOD %d Data from Skeletal Mesh %ls.", LODIndex, *SkeletalMesh->GetName());
 					}
 					return;
 				}
@@ -2007,7 +2007,7 @@ namespace UnrealConversionUtils
 						if (NumClothLOD > 1)
 						{
 							NumClothLOD = 1;
-							UE_LOG(LogMutable, Error, TEXT("\"Mappings to Same LOD\" only supported in Ray Tracing \"Cloth LODBias Mode\" option for [%s] Skeletal Mesh"), *SkeletalMesh->GetName());
+							UE_LOGF(LogMutable, Error, "\"Mappings to Same LOD\" only supported in Ray Tracing \"Cloth LODBias Mode\" option for [%ls] Skeletal Mesh", *SkeletalMesh->GetName());
 						}
 						
 						uint32 NumVertices = SectionClothIndexMapping.LODBiasStride * NumClothLOD;
diff --git a/Engine/Plugins/Mutable/Source/CustomizableObject/Private/MuCO/UnrealMutableImageProvider.cpp b/Engine/Plugins/Mutable/Source/CustomizableObject/Private/MuCO/UnrealMutableImageProvider.cpp
--- a/Engine/Plugins/Mutable/Source/CustomizableObject/Private/MuCO/UnrealMutableImageProvider.cpp
+++ b/Engine/Plugins/Mutable/Source/CustomizableObject/Private/MuCO/UnrealMutableImageProvider.cpp
@@ -143,7 +143,7 @@ TTuple<UE::Tasks::FTask, TFunction<void()>> FUnrealMutableResourceProvider::GetI
 
 	if (!Texture)
 	{
-		UE_LOG(LogMutable, Warning, TEXT("Invalid Image Parameter. Nullptr"));
+		UE_LOGF(LogMutable, Warning, "Invalid Image Parameter. Nullptr");
 		ResultCallback(CreateDummy());
 		return Invoke(TrivialReturn);
 	}
@@ -151,7 +151,7 @@ TTuple<UE::Tasks::FTask, TFunction<void()>> FUnrealMutableResourceProvider::GetI
 	UTexture2D* Texture2D = Cast<UTexture2D>(Texture);
 	if (!Texture2D)
 	{
-		UE_LOG(LogMutable, Warning, TEXT("Invalid Image Parameter [%s]. Is not a UTexture2D."), *Texture->GetName());
+		UE_LOGF(LogMutable, Warning, "Invalid Image Parameter [%ls]. Is not a UTexture2D.", *Texture->GetName());
 		ResultCallback(CreateDummy());
 		return Invoke(TrivialReturn);
 	}
@@ -171,7 +171,7 @@ TTuple<UE::Tasks::FTask, TFunction<void()>> FUnrealMutableResourceProvider::GetI
 	{
 		// This could happen in the editor, because some source textures may have changed while there was a background compilation.
 		// We just show a warning and move on. This cannot happen during cooks, so it is fine.
-		UE_LOG(LogMutable, Warning, TEXT("Failed to load some source texture data for image [%s]. Some materials may look corrupted."), *Texture->GetName());
+		UE_LOGF(LogMutable, Warning, "Failed to load some source texture data for image [%ls]. Some materials may look corrupted.", *Texture->GetName());
 	}
 
 	ResultCallback(Image);
@@ -189,7 +189,7 @@ TTuple<UE::Tasks::FTask, TFunction<void()>> FUnrealMutableResourceProvider::GetI
 	// Check if it's a format we support
 	if (MutImageFormat == UE::Mutable::Private::EImageFormat::None)
 	{
-		UE_LOG(LogMutable, Warning, TEXT("Failed to get Image Parameter [%s]. Unexpected image format. EImageFormat [%s]."), *Texture2D->GetName(), GetPixelFormatString(Format));
+		UE_LOGF(LogMutable, Warning, "Failed to get Image Parameter [%ls]. Unexpected image format. EImageFormat [%ls].", *Texture2D->GetName(), GetPixelFormatString(Format));
 		ResultCallback(CreateDummy());
 		return Invoke(TrivialReturn);
 	}
@@ -222,7 +222,7 @@ TTuple<UE::Tasks::FTask, TFunction<void()>> FUnrealMutableResourceProvider::GetI
 
 	if (!CheckBulkSizesAreValid(Texture2D, Image.Get(), TextureMipIndexBegin, TextureMipIndexEnd))
 	{
-		UE_LOG(LogMutable, Warning, TEXT("Failed to get Image Parameter [%s]. Bulk data does not match expected image data sizes."), *Texture2D->GetName());
+		UE_LOGF(LogMutable, Warning, "Failed to get Image Parameter [%ls]. Bulk data does not match expected image data sizes.", *Texture2D->GetName());
 		ResultCallback(CreateDummy());
 
 		return Invoke(TrivialReturn);
@@ -263,14 +263,14 @@ TTuple<UE::Tasks::FTask, TFunction<void()>> FUnrealMutableResourceProvider::GetI
 		// Should we do someting different than returning a dummy image if cancelled?
 		if (Status == FBulkDataRequest::EStatus::Cancelled)
 		{
-			UE_LOG(LogMutable, Warning, TEXT("Failed to get external image. Cancelled Bulk Data Request"));
+			UE_LOGF(LogMutable, Warning, "Failed to get external image. Cancelled Bulk Data Request");
 			ResultCallback(CreateDummy());
 			return;
 		}
 		
 		if (Status != FBulkDataRequest::EStatus::Ok)
 		{
-			UE_LOG(LogMutable, Warning, TEXT("Failed to get external image parameter."));
+			UE_LOGF(LogMutable, Warning, "Failed to get external image parameter.");
 			ResultCallback(CreateDummy());
 			return;
 		}
@@ -301,7 +301,7 @@ TTuple<UE::Tasks::FTask, TFunction<void()>> FUnrealMutableResourceProvider::GetR
 	{
 		// This could happen in the editor, because some source textures may have changed while there was a background compilation.
 		// We just show a warning and move on. This cannot happen during cooks, so it is fine.
-		UE_LOG(LogMutable, Warning, TEXT("Failed to load Referenced Image [%i]. Invalid id."), Id);
+		UE_LOGF(LogMutable, Warning, "Failed to load Referenced Image [%i]. Invalid id.", Id);
 		ResultCallback(CreateDummy());
 		return Invoke(TrivialReturn);
 	}
@@ -309,7 +309,7 @@ TTuple<UE::Tasks::FTask, TFunction<void()>> FUnrealMutableResourceProvider::GetR
 	UTexture* Texture = Images[Id].Get();
 	if (!Texture)
 	{
-		UE_LOG(LogMutable, Warning, TEXT("Invalid Referenced Image [%i]. Nullptr."), Id);
+		UE_LOGF(LogMutable, Warning, "Invalid Referenced Image [%i]. Nullptr.", Id);
 		ResultCallback(CreateDummy());
 		return Invoke(TrivialReturn);
 	}
@@ -317,7 +317,7 @@ TTuple<UE::Tasks::FTask, TFunction<void()>> FUnrealMutableResourceProvider::GetR
 	UTexture2D* Texture2D = Cast<UTexture2D>(Texture);
 	if (!Texture2D)
 	{
-		UE_LOG(LogMutable, Warning, TEXT("Invalid Referenced Image [%i, %s]. Is not a UTexture2D."), Id, *Texture->GetName());
+		UE_LOGF(LogMutable, Warning, "Invalid Referenced Image [%i, %ls]. Is not a UTexture2D.", Id, *Texture->GetName());
 		ResultCallback(CreateDummy());
 		return Invoke(TrivialReturn);
 	}
@@ -335,7 +335,7 @@ TTuple<UE::Tasks::FTask, TFunction<void()>> FUnrealMutableResourceProvider::GetR
 	{
 		// This could happen in the editor, because some source textures may have changed while updating.
 		// We just show a warning and move on. This cannot happen during cooks, so it is fine.
-		UE_LOG(LogMutable, Warning, TEXT("Failed to load some source texture data for Referenced Image [%i, %s]. Some textures may be corrupted."), Id, *Texture->GetName());
+		UE_LOGF(LogMutable, Warning, "Failed to load some source texture data for Referenced Image [%i, %ls]. Some textures may be corrupted.", Id, *Texture->GetName());
 		
 		ResultCallback(CreateDummy());
 		return Invoke(TrivialReturn);
@@ -346,7 +346,7 @@ TTuple<UE::Tasks::FTask, TFunction<void()>> FUnrealMutableResourceProvider::GetR
 #else // WITH_EDITOR
 
 	// Not supported outside editor yet.
-	UE_LOG(LogMutable, Warning, TEXT("Failed to get Reference Image. Only supported in editor."));
+	UE_LOGF(LogMutable, Warning, "Failed to get Reference Image. Only supported in editor.");
 
 	ResultCallback(CreateDummy());
 	return Invoke(TrivialReturn);
@@ -403,7 +403,7 @@ UE::Mutable::Private::FExtendedImageDesc FUnrealMutableResourceProvider::GetImag
 	// Check if it's a format we support
 	if (MutableFormat == UE::Mutable::Private::EImageFormat::None)
 	{
-		UE_LOG(LogMutable, Warning, TEXT("Failed to get Image Parameter descriptor. Unexpected image format. EImageFormat [%s]."), GetPixelFormatString(Format));
+		UE_LOGF(LogMutable, Warning, "Failed to get Image Parameter descriptor. Unexpected image format. EImageFormat [%ls].", GetPixelFormatString(Format));
 		return CreateDummyDesc();
 	}
 
@@ -469,12 +469,12 @@ void FUnrealMutableResourceProvider::CacheRuntimeReferencedImages(const TArray<T
 		UTexture2D* Texture = RuntimeReferencedTexture.Get(); // Is already loaded.
 		if (!Texture)
 		{
-			UE_LOG(LogMutable, Warning, TEXT("Runtime Referenced Texture [%s] was not async loaded. Forcing load sync."), *RuntimeReferencedTexture->GetPathName());
+			UE_LOGF(LogMutable, Warning, "Runtime Referenced Texture [%ls] was not async loaded. Forcing load sync.", *RuntimeReferencedTexture->GetPathName());
 			
 			Texture = UE::Mutable::Private::LoadObject(RuntimeReferencedTexture);
 			if (!Texture)
 			{
-				UE_LOG(LogMutable, Warning, TEXT("Failed to force load sync [%s]."), *RuntimeReferencedTexture->GetPathName());
+				UE_LOGF(LogMutable, Warning, "Failed to force load sync [%ls].", *RuntimeReferencedTexture->GetPathName());
 				continue;
 			}
 		}
diff --git a/Engine/Plugins/Mutable/Source/CustomizableObject/Private/MuCO/UnrealMutableModelDiskStreamer.cpp b/Engine/Plugins/Mutable/Source/CustomizableObject/Private/MuCO/UnrealMutableModelDiskStreamer.cpp
--- a/Engine/Plugins/Mutable/Source/CustomizableObject/Private/MuCO/UnrealMutableModelDiskStreamer.cpp
+++ b/Engine/Plugins/Mutable/Source/CustomizableObject/Private/MuCO/UnrealMutableModelDiskStreamer.cpp
@@ -291,7 +291,7 @@ UE::Mutable::Private::FModelReader::FOperationID FUnrealMutableModelBulkReader::
 
 	FScopeLock Lock(&CriticalSection);
 	
-	UE_LOG(LogMutable, VeryVerbose, TEXT("Streaming: reading data %08llu."), Key);
+	UE_LOGF(LogMutable, VeryVerbose, "Streaming: reading data %08llu.", Key);
 
 	// Find the object we are streaming for
 	FObjectData* ObjectData = Objects.FindByPredicate(
@@ -312,7 +312,7 @@ UE::Mutable::Private::FModelReader::FOperationID FUnrealMutableModelBulkReader::
 	const TSharedPtr<FModelStreamableBulkData> PinnedModelStreamableBulkData = ObjectData->ModelStreamableBulkData.Pin();	
 	if (!PinnedModelStreamableBulkData)
 	{
-		UE_LOG(LogMutable, Error, TEXT("Streamable Bulk Data not available!"));
+		UE_LOGF(LogMutable, Error, "Streamable Bulk Data not available!");
 		if (CompletionCallback)
 		{
 			(*CompletionCallback)(false);
@@ -324,7 +324,7 @@ UE::Mutable::Private::FModelReader::FOperationID FUnrealMutableModelBulkReader::
 	if (!Block)
 	{
 		// File Handle not found! This shouldn't really happen.
-		UE_LOG(LogMutable, Error, TEXT("Streaming Block not found!"));
+		UE_LOGF(LogMutable, Error, "Streaming Block not found!");
 		check(false);
 
 		if (CompletionCallback)
@@ -441,7 +441,7 @@ UE::Mutable::Private::FModelReader::FOperationID FUnrealMutableModelBulkReader::
 			
 			if (!FileHandle)
 			{
-				UE_LOG(LogMutable, Error, TEXT("Failed to create AsyncReadFileHandle. File Path [%s]."), *FilePath);
+				UE_LOGF(LogMutable, Error, "Failed to create AsyncReadFileHandle. File Path [%ls].", *FilePath);
 				check(false);
 				return -1;
 			}
@@ -494,7 +494,7 @@ bool FUnrealMutableModelBulkReader::IsReadCompleted(UE::Mutable::Private::FModel
 		}
 	}
 
-	UE_LOG(LogMutable, Error, TEXT("Operation not found in IsReadCompleted."));
+	UE_LOGF(LogMutable, Error, "Operation not found in IsReadCompleted.");
 	check(false);
 	return true;
 }
@@ -537,7 +537,7 @@ bool FUnrealMutableModelBulkReader::EndRead(UE::Mutable::Private::FModelReader::
 #endif
 			if (!bCompleted)
 			{
-				UE_LOG(LogMutable, Error, TEXT("Operation failed to complete in EndRead."));
+				UE_LOGF(LogMutable, Error, "Operation failed to complete in EndRead.");
 				check(false);
 			}
 
@@ -551,7 +551,7 @@ bool FUnrealMutableModelBulkReader::EndRead(UE::Mutable::Private::FModelReader::
 
 	if (!bFound)
 	{
-		UE_LOG(LogMutable, Error, TEXT("Operation not found in EndRead."));
+		UE_LOGF(LogMutable, Error, "Operation not found in EndRead.");
 		check(false);
 	}
 
diff --git a/Engine/Plugins/Mutable/Source/CustomizableObject/Private/MuCO/UnrealToMutableTextureConversionUtils.cpp b/Engine/Plugins/Mutable/Source/CustomizableObject/Private/MuCO/UnrealToMutableTextureConversionUtils.cpp
--- a/Engine/Plugins/Mutable/Source/CustomizableObject/Private/MuCO/UnrealToMutableTextureConversionUtils.cpp
+++ b/Engine/Plugins/Mutable/Source/CustomizableObject/Private/MuCO/UnrealToMutableTextureConversionUtils.cpp
@@ -393,7 +393,7 @@ EUnrealToMutableConversionError ConvertTextureUnrealSourceToMutable(UE::Mutable:
 		}
 
 		//FString Msg = FString::Printf(TEXT("Alpha channel is %s for %s"), bHasAlphaChannel ? TEXT("enabled") : TEXT("disabled"), *Texture->GetName());
-		//UE_LOG(LogMutable, VeryVerbose, TEXT("%s"), *Msg);
+		//UE_LOGF(LogMutable, VeryVerbose, "%ls", *Msg);
 
 		break;
 	}
