[Mutable] Optimize OpImageLayer zero mask optimization by vectorizing the test.

#rnx

#Synced-CL 51758072

diff --git a/Engine/Plugins/Mutable/Source/MutableRuntime/Private/MuR/ImageComputeUtils.cpp b/Engine/Plugins/Mutable/Source/MutableRuntime/Private/MuR/ImageComputeUtils.cpp
--- a/Engine/Plugins/Mutable/Source/MutableRuntime/Private/MuR/ImageComputeUtils.cpp
+++ b/Engine/Plugins/Mutable/Source/MutableRuntime/Private/MuR/ImageComputeUtils.cpp
@@ -35,6 +35,11 @@ namespace ImageComputeUtils
 		return _mm_loadu_si128((__m128i*)Ptr);
 	}
 
+	FORCEINLINE bool TestZero_Int8x16(VectorRegisterInt8x16 Reg)
+	{
+		return _mm_testz_si128(Reg, Reg) != 0; 
+	}
+
 	FORCEINLINE VectorRegisterInt8x16 Shuffle_Int8x16(const VectorRegisterInt8x16 V, const VectorRegisterInt8x16 Shuffle)
 	{
 		return _mm_shuffle_epi8(V, Shuffle);
@@ -50,6 +55,7 @@ namespace ImageComputeUtils
 		return _mm_blendv_epi8(B, A, Mask);
 	}
 
+
 	FORCEINLINE VectorRegisterInt8x16 Unorm255Lerp_Int8x16(const VectorRegisterInt8x16 A, const VectorRegisterInt8x16 B, const VectorRegisterInt8x16 Factor)
 	{
 		__m128i OneMinusFactor = _mm_xor_si128(Factor, _mm_set1_epi8((uint8)0xFF)); 
@@ -95,6 +101,12 @@ namespace ImageComputeUtils
 		return vld1q_u8(Ptr);
 	}
 
+	FORCEINLINE bool TestZero_Int8x16(VectorRegisterUInt8x16 Reg)
+	{
+		uint64x2_t Val = vreinterpretq_u64_u8(Reg);
+		return (vgetq_lane_u64(Val, 0) | vgetq_lane_u64(Val, 1)) == 0;
+	}
+
 	FORCEINLINE VectorRegisterInt8x16 Shuffle_Int8x16(VectorRegisterInt8x16 V, VectorRegisterInt8x16 Shuffle)
 	{
 		return vqtbl1q_u8(V, Shuffle);
@@ -618,5 +630,36 @@ namespace ImageComputeUtils
 			Dest[ElemIndex] = static_cast<uint8>((Value * 257) >> 16);
 		}
 	}
+
+
+	bool TestZero_U8(const uint8* Data, int32 NumElems)
+	{
+		int32 ElemIndex = 0;
+
+#if UE_MUTABLE_IMAGE_COMPUTE_UTILS_USE_VECTOR_IMPL
+		constexpr int32 ElemsPerIter = 16;
+		int32 NumIters = NumElems / ElemsPerIter;
+
+		for (int32 Iter = 0; Iter < NumIters; ++Iter)
+		{	
+			if (!TestZero_Int8x16(LoadUnaligned_Int8x16(Data + Iter*ElemsPerIter)))
+			{
+				return false;
+			}
+		}
+
+		ElemIndex = NumIters*ElemsPerIter; 
+#endif
+
+		for (; ElemIndex < NumElems; ++ElemIndex)
+		{
+			if (Data[ElemIndex] != 0)
+			{
+				return false;
+			}
+		}
+
+		return true;
+	}
 } // namespace ImageComputeUtils
 } // namespace UE::Mutable::Private
diff --git a/Engine/Plugins/Mutable/Source/MutableRuntime/Private/MuR/ImageComputeUtils.h b/Engine/Plugins/Mutable/Source/MutableRuntime/Private/MuR/ImageComputeUtils.h
--- a/Engine/Plugins/Mutable/Source/MutableRuntime/Private/MuR/ImageComputeUtils.h
+++ b/Engine/Plugins/Mutable/Source/MutableRuntime/Private/MuR/ImageComputeUtils.h
@@ -36,5 +36,6 @@ namespace ImageComputeUtils
 	void Vec3Fill_U16(uint16* Dest, Math::TIntVector4<uint16> Value, int32 NumElems);	
 	void Vec1Fill_U16(uint16* Dest, Math::TIntVector4<uint16> Value, int32 NumElems);
 
+	bool TestZero_U8(const uint8* Data, int32 NumElems);
 } // namespace ImageComputeUtils
 } // namespace UE::Mutable::Private
diff --git a/Engine/Plugins/Mutable/Source/MutableRuntime/Private/MuR/OpImageLayer.cpp b/Engine/Plugins/Mutable/Source/MutableRuntime/Private/MuR/OpImageLayer.cpp
--- a/Engine/Plugins/Mutable/Source/MutableRuntime/Private/MuR/OpImageLayer.cpp
+++ b/Engine/Plugins/Mutable/Source/MutableRuntime/Private/MuR/OpImageLayer.cpp
@@ -1637,7 +1637,7 @@ namespace OpImageLayerInternal
 					MaskDataBuffer = Register4.Data;
 				}
 
-				bSkipColorBlend = FMemory::MemIsZero(MaskDataBuffer, NumWaveElems*BaseBytesPerElem);
+				bSkipColorBlend = ImageComputeUtils::TestZero_U8(MaskDataBuffer, NumWaveElems*BaseBytesPerElem);
 			}
 			
 			if (Kernel.BaseFormat)
