[mutable] Fixed concurrency issue in Mesh and Texture mutable caches
- Removed no longer used includes.

#prefligh 69b9279b7663528f59558db0

#rnx
#Synced-CL 51817261

diff --git a/Engine/Plugins/Mutable/Source/CustomizableObject/Internal/MuCO/CustomizableObjectPrivate.h b/Engine/Plugins/Mutable/Source/CustomizableObject/Internal/MuCO/CustomizableObjectPrivate.h
--- a/Engine/Plugins/Mutable/Source/CustomizableObject/Internal/MuCO/CustomizableObjectPrivate.h
+++ b/Engine/Plugins/Mutable/Source/CustomizableObject/Internal/MuCO/CustomizableObjectPrivate.h
@@ -136,7 +136,7 @@ struct FMutableModelParameterProperties
 class FMeshCache
 {
 public:
-	TWeakObjectPtr<UCustomizableObjectSkeletalMesh> Get(const TArray<UE::Mutable::Private::FMeshId>& Key);
+	TStrongObjectPtr<UCustomizableObjectSkeletalMesh> Get(const TArray<UE::Mutable::Private::FMeshId>& Key);
 
 	void Add(const TArray<UE::Mutable::Private::FMeshId>& Key, UCustomizableObjectSkeletalMesh* Value);
 
@@ -171,7 +171,7 @@ public:
 		friend uint32 GetTypeHash(const FId& Key);
 	};
 	
-	TWeakObjectPtr<UTexture2D> Get(const FId& Key);
+	TStrongObjectPtr<UTexture2D> Get(const FId& Key);
 
 	void Add(const FId& Key, UTexture2D* Value);
 
diff --git a/Engine/Plugins/Mutable/Source/CustomizableObject/Private/MuCO/CustomizableObject.cpp b/Engine/Plugins/Mutable/Source/CustomizableObject/Private/MuCO/CustomizableObject.cpp
--- a/Engine/Plugins/Mutable/Source/CustomizableObject/Private/MuCO/CustomizableObject.cpp
+++ b/Engine/Plugins/Mutable/Source/CustomizableObject/Private/MuCO/CustomizableObject.cpp
@@ -2,14 +2,11 @@
 
 #include "MuCO/CustomizableObject.h"
 
-#include "Algo/Copy.h"
-#include "Async/AsyncFileHandle.h"
 #include "EdGraph/EdGraph.h"
 #include "Engine/Engine.h"
 #include "Engine/SkeletalMesh.h"
 #include "Animation/AnimInstance.h"
 #include "Animation/Skeleton.h"
-#include "Engine/AssetUserData.h"
 #include "HAL/FileManager.h"
 #include "HAL/PlatformFileManager.h"
 #include "Input/Reply.h"
@@ -27,7 +24,6 @@
 #include "MuCO/MutableProjectorTypeUtils.h"
 #include "MuCO/UnrealMutableModelDiskStreamer.h"
 #include "MuR/Model.h"
-#include "MuR/MutableEditorLogger.h"
 #include "MuR/Operations.h"
 #include "PhysicsEngine/PhysicsAsset.h"
 #include "Serialization/MemoryReader.h"
@@ -46,7 +42,6 @@
 #include "DerivedDataCacheKey.h"
 #include "DerivedDataRequestOwner.h"
 #include "MuCO/LoadUtils.h"
-#include "AssetToolsModule.h"
 #endif
 
 #include "ExternalOperationProvider.h"
@@ -109,14 +104,14 @@ FGuid UCustomizableObjectPrivate::GetVersionId() const
 
 void UCustomizableObject::GetAssetRegistryTags(FAssetRegistryTagsContext Context) const
 {
-	int32 isRoot = 0;
+	int32 bIsRoot = 0;
 
 	if (const ICustomizableObjectEditorModule* Module = ICustomizableObjectEditorModule::Get())
 	{
-		isRoot = Module->IsRootObject(*this) ? 1 : 0;
+		bIsRoot = Module->IsRootObject(*this) ? 1 : 0;
 	}
 	
-	Context.AddTag(FAssetRegistryTag("IsRoot", FString::FromInt(isRoot), FAssetRegistryTag::TT_Numerical));
+	Context.AddTag(FAssetRegistryTag("IsRoot", FString::FromInt(bIsRoot), FAssetRegistryTag::TT_Numerical));
 	Super::GetAssetRegistryTags(Context);
 }
 
@@ -1819,12 +1814,18 @@ void UCustomizableObjectPrivate::ApplyStateForcedValuesToParameters(FCustomizabl
 //-------------------------------------------------------------------------------------------------
 //-------------------------------------------------------------------------------------------------
 
-TWeakObjectPtr<UCustomizableObjectSkeletalMesh> FMeshCache::Get(const TArray<UE::Mutable::Private::FMeshId>& Key)
+TStrongObjectPtr<UCustomizableObjectSkeletalMesh> FMeshCache::Get(const TArray<UE::Mutable::Private::FMeshId>& Key)
 {
 	UE::TScopeLock<UE::FMutex> LockGuard(Lock);
+
+	const TWeakObjectPtr<UCustomizableObjectSkeletalMesh>* FoundMesh = GeneratedMeshes->Find(Key);
+	if (!FoundMesh)
+	{
+		return TStrongObjectPtr<UCustomizableObjectSkeletalMesh>();
+	}
 	
-	const TWeakObjectPtr<UCustomizableObjectSkeletalMesh>* Result = GeneratedMeshes->Find(Key);
-	return Result ? *Result : TWeakObjectPtr<UCustomizableObjectSkeletalMesh>();
+	bool bWasPinSuccessful = false;
+	return FoundMesh->TryPin(bWasPinSuccessful);
 }
 
 
@@ -1883,12 +1884,18 @@ void FSkeletonCache::Clear()
 }
 
 
-TWeakObjectPtr<UTexture2D> FTextureCache::Get(const FId& Key)
+TStrongObjectPtr<UTexture2D> FTextureCache::Get(const FId& Key)
 {
 	UE::TScopeLock<UE::FMutex> LockGuard(Lock);
 
-	TWeakObjectPtr<UTexture2D>* FoundTexture = GeneratedTextures.Find(Key);
-	return FoundTexture ? *FoundTexture : TWeakObjectPtr<UTexture2D>();
+	const TWeakObjectPtr<UTexture2D>* FoundTexture = GeneratedTextures.Find(Key);
+	if (!FoundTexture)
+	{
+		return TStrongObjectPtr<UTexture2D>();
+	}
+	
+	bool bWasPinSuccessful = false;
+	return FoundTexture->TryPin(bWasPinSuccessful);
 }
 
 
diff --git a/Engine/Plugins/Mutable/Source/CustomizableObject/Private/MuCO/CustomizableObjectInstance.cpp b/Engine/Plugins/Mutable/Source/CustomizableObject/Private/MuCO/CustomizableObjectInstance.cpp
--- a/Engine/Plugins/Mutable/Source/CustomizableObject/Private/MuCO/CustomizableObjectInstance.cpp
+++ b/Engine/Plugins/Mutable/Source/CustomizableObject/Private/MuCO/CustomizableObjectInstance.cpp
@@ -38,7 +38,6 @@
 #include "RenderingThread.h"
 #include "SkeletalMergingLibrary.h"
 #include "MuCO/ICustomizableObjectEditorModule.h"
-#include "UObject/UObjectIterator.h"
 #include "Widgets/Notifications/SNotificationList.h"
 #include "PhysicsEngine/AggregateGeom.h"
 #include "PhysicsEngine/PhysicsConstraintTemplate.h"
@@ -47,7 +46,6 @@
 #include "MuCO/LoadUtils.h"
 #include "Serialization/BulkData.h"
 
-#include "MuCO/Plugins/IMutableClothingModule.h"
 #include "Rendering/SkeletalMeshModel.h"
 
 #include "MuR/MutableEditorLogger.h"
@@ -1769,13 +1767,12 @@ bool UCustomizableInstancePrivate::UpdateSkeletalMesh_PostBeginUpdate0(UCustomiz
 
 		if (Context->bUseMeshCache)
 		{
-			TWeakObjectPtr<UCustomizableObjectSkeletalMesh> MeshCacheEntry = CustomizableObject->GetPrivate()->MeshCache.Get(
+			TStrongObjectPtr<UCustomizableObjectSkeletalMesh> CachedMesh = CustomizableObject->GetPrivate()->MeshCache.Get(
 				ComponentDescriptor->MeshDescriptors);
 			
-			TObjectPtr<UCustomizableObjectSkeletalMesh> Mesh = MeshCacheEntry.Get();
-			if (IsValid(Mesh))
+			if (CachedMesh)
 			{
-				SkeletalMeshes.Add(ComponentName, Mesh);
+				SkeletalMeshes.Add(ComponentName, CachedMesh.Get());
 				continue;
 			}
 		}
@@ -4304,8 +4301,8 @@ UTexture* BuildTexture(const TSharedRef<FUpdateContextPrivate>& Context, const i
 	{
 		if (Context->bUseCaches)
 		{
-			const TWeakObjectPtr<UTexture2D> WeakTexture = TextureCache.Get(TextureCacheKey);
-			Texture = WeakTexture.Get();
+			const TStrongObjectPtr<UTexture2D> FoundTexture = TextureCache.Get(TextureCacheKey);
+			Texture = FoundTexture.Get();
 			check(Texture);
 		}
 	}
diff --git a/Engine/Plugins/Mutable/Source/CustomizableObject/Private/MuCO/CustomizableObjectSystem.cpp b/Engine/Plugins/Mutable/Source/CustomizableObject/Private/MuCO/CustomizableObjectSystem.cpp
--- a/Engine/Plugins/Mutable/Source/CustomizableObject/Private/MuCO/CustomizableObjectSystem.cpp
+++ b/Engine/Plugins/Mutable/Source/CustomizableObject/Private/MuCO/CustomizableObjectSystem.cpp
@@ -37,14 +37,11 @@
 #include "Logging/MessageLog.h"
 #include "Misc/ConfigCacheIni.h"
 #include "Engine/World.h"
-#include "Interfaces/ITargetPlatform.h"
 #include "MuT/UnrealPixelFormatOverride.h"
-#include "IMessageLogListing.h"
 #else
 #include "Engine/Engine.h"
 #endif
 
-#include "MutableStreamRequest.h"
 #include "MuCO/CustomizableObjectSkeletalMesh.h"
 #include "MuCO/LoadUtils.h"
 
@@ -1906,10 +1903,8 @@ namespace impl
 		{
 			for (const FComponentDescriptor& ComponentDescriptor : OperationData->GetComponentDescriptors())
 			{
-				TWeakObjectPtr<UCustomizableObjectSkeletalMesh> MeshCacheEntry = OperationData->Object->GetPrivate()->MeshCache.Get(ComponentDescriptor.MeshDescriptors);
-				
-				bool bWasPinningSuccessful = false;
-				if (TStrongObjectPtr<UCustomizableObjectSkeletalMesh> CachedMesh = MeshCacheEntry.TryPin(bWasPinningSuccessful))
+				TStrongObjectPtr<UCustomizableObjectSkeletalMesh> CachedMesh = OperationData->Object->GetPrivate()->MeshCache.Get(ComponentDescriptor.MeshDescriptors);
+				if (CachedMesh)
 				{
 					OperationData->Objects.Emplace(CachedMesh);
 				}
@@ -2127,11 +2122,8 @@ namespace impl
 				{
 					FTextureCache& TextureCache = OperationData->Object->GetPrivate()->TextureCache;
 					const FTextureCache::FId TextureCacheKey = FTextureCache::FId(Image->ImageID, OperationData->MipsToSkip, OperationData->bBake);
-					const TWeakObjectPtr<UTexture2D> CachedWeakTexture = TextureCache.Get(TextureCacheKey);
-						
-					bool bWasPinSuccessful = false;
-					TStrongObjectPtr<const UTexture> CachedTexture = CachedWeakTexture.TryPin(bWasPinSuccessful);
-					if (CachedTexture.IsValid())
+
+					if (const TStrongObjectPtr<const UTexture> CachedTexture = TextureCache.Get(TextureCacheKey))
 					{
 						UE_LOGF(LogMutable, VeryVerbose, "Texture resource with id [%ls] is cached.", *Image->ImageID.ToString());
 						OperationData->Objects.Add(CachedTexture);
