[mutable] Added !Mutable Skeletal Mesh make"  and "Skeletal Mesh Make" nodes.
- Updated mutable SKM Component node so it has as input, instead of the LOD array pins, a single passthrough SKM pin.
- Added a new node named "Skeletal Mesh Make" which will from now on define the LOD configurations that were previously part of the component node.
- Added a new node named "Mutable Skeletal Mesh Make" which will now contain the list of LODs that the component node was hosting before.
- Added backwards compatibility code to the SKM Component node so all the data before hosted there gets migrated to the new nodes.
- Made the addition and removal of LOD pins from the "Mutable Skeletal Mesh Make" use the "Add Pin" procedure instead of the details aproach.
- Updated First Pass code so it makes the Component node reach for the data found in the connected nodes. This approach is just temporary while the GenerateMutableSource code is being implemented.
- The GeneateMutableSource for the component code has been slightly altered to make use of the new nodes but that code is just temporary, just to keep the graph working while we work on the final implementation of the generation code. On the practical side of things, all should work as before, as no new code compilation nor ASTOps have been made for the new nodes yet.
- Other tweaks and changes

#rnx
#Synced-CL 51817239

diff --git a/Engine/Plugins/Mutable/Source/CustomizableObject/Internal/MuCO/CustomizableObjectCustomVersion.h b/Engine/Plugins/Mutable/Source/CustomizableObject/Internal/MuCO/CustomizableObjectCustomVersion.h
--- a/Engine/Plugins/Mutable/Source/CustomizableObject/Internal/MuCO/CustomizableObjectCustomVersion.h
+++ b/Engine/Plugins/Mutable/Source/CustomizableObject/Internal/MuCO/CustomizableObjectCustomVersion.h
@@ -232,6 +232,8 @@ struct FCustomizableObjectCustomVersion
 		
 		RemovedTypedSwitches,
 		
+		SkeletalMeshMakeNodeDeprecations,				// SKM Component has gotten some responsibilities stripped from it and moved to the new Skeletal Mesh Make node
+		
 		// -----<new versions can be added above this line>-------------------------------------------------
 		VersionPlusOne,
 		LatestVersion = VersionPlusOne - 1
diff --git a/Engine/Plugins/Mutable/Source/CustomizableObjectEditor/Private/MuCOE/CustomizableObjectEditorModule.cpp b/Engine/Plugins/Mutable/Source/CustomizableObjectEditor/Private/MuCOE/CustomizableObjectEditorModule.cpp
--- a/Engine/Plugins/Mutable/Source/CustomizableObjectEditor/Private/MuCOE/CustomizableObjectEditorModule.cpp
+++ b/Engine/Plugins/Mutable/Source/CustomizableObjectEditor/Private/MuCOE/CustomizableObjectEditorModule.cpp
@@ -70,7 +70,6 @@
 #include "MuCOE/Nodes/CustomizableObjectNodeTable.h"
 #include "MuCOE/Nodes/CustomizableObjectNodeTableDetails.h"
 #include "MuCOE/Nodes/CustomizableObjectNodeTunnel.h"
-#include "MuCOE/Nodes/CustomizableObjectNodeComponentMeshDetails.h"
 #include "MuCOE/Nodes/CustomizableObjectNodeVariation.h"
 #include "MuCOE/Nodes/CustomizableObjectNodeTexture.h"
 #include "MuCOE/Nodes/CustomizableObjectNodeTextureVariation.h"
@@ -107,6 +106,7 @@
 #include "Nodes/ExternalTypeCustomization.h"
 #include "PhysicsEngine/PhysicsAsset.h"
 #include "Components/SkeletalMeshComponent.h"
+#include "Nodes/CONodeMutableSkeletalMeshMake.h"
 
 class AActor;
 class FString;
@@ -267,7 +267,8 @@ void FCustomizableObjectEditorModule::StartupModule()
 	RegisterCustomDetails(PropertyModule, UCustomizableObjectNodeStaticMesh::StaticClass(), FOnGetDetailCustomizationInstance::CreateStatic(&FCustomizableObjectNodeDetails::MakeInstance));
 	RegisterCustomDetails(PropertyModule, UCustomizableObjectNodeSkeletalMeshParameter::StaticClass(), FOnGetDetailCustomizationInstance::CreateStatic(&FCustomizableObjectNodeSkeletalMeshParameterDetails::MakeInstance));
 	RegisterCustomDetails(PropertyModule, UCustomizableObjectNodeTable::StaticClass(), FOnGetDetailCustomizationInstance::CreateStatic(&FCustomizableObjectNodeTableDetails::MakeInstance));
-	RegisterCustomDetails(PropertyModule, UCustomizableObjectNodeComponentMesh::StaticClass(), FOnGetDetailCustomizationInstance::CreateStatic(&FCustomizableObjectNodeComponentMeshDetails::MakeInstance));
+	RegisterCustomDetails(PropertyModule, UCustomizableObjectNodeComponentMesh::StaticClass(), FOnGetDetailCustomizationInstance::CreateStatic(&FCustomizableObjectNodeDetails::MakeInstance));
+	RegisterCustomDetails(PropertyModule, UCONodeMutableSkeletalMeshMake::StaticClass(), FOnGetDetailCustomizationInstance::CreateStatic(&FCustomizableObjectNodeDetails::MakeInstance));
 	RegisterCustomDetails(PropertyModule, UCustomizableObjectNodeMacroInstance::StaticClass(), FOnGetDetailCustomizationInstance::CreateStatic(&FCustomizableObjectNodeMacroInstanceDetails::MakeInstance));
 	RegisterCustomDetails(PropertyModule, UCustomizableObjectNodeParameter::StaticClass(), FOnGetDetailCustomizationInstance::CreateStatic(&FCustomizableObjectNodeParameterDetails::MakeInstance));
 	RegisterCustomDetails(PropertyModule, UCONodeMaterialBreak::StaticClass(), FOnGetDetailCustomizationInstance::CreateStatic(&FCustomizableObjectNodeDetails::MakeInstance));
@@ -289,7 +290,6 @@ void FCustomizableObjectEditorModule::StartupModule()
 	// Custom properties
 	PropertyModule.RegisterCustomPropertyTypeLayout("CustomizableObjectIdentifier", FOnGetPropertyTypeCustomizationInstance::CreateStatic(&FCustomizableObjectIdentifierCustomization::MakeInstance));
 	PropertyModule.RegisterCustomPropertyTypeLayout(MeshReshapeBoneReferenceUStructName, FOnGetPropertyTypeCustomizationInstance::CreateStatic(&FMeshReshapeBonesReferenceCustomization::MakeInstance));
-	PropertyModule.RegisterCustomPropertyTypeLayout(BoneToRemoveUStructName, FOnGetPropertyTypeCustomizationInstance::CreateStatic(&FCustomizableObjectLODReductionSettings::MakeInstance));
 	PropertyModule.RegisterCustomPropertyTypeLayout(NAME_StrProperty, FOnGetPropertyTypeCustomizationInstance::CreateStatic(&FCustomizableObjectStateParameterSelector::MakeInstance), MakeShared<FStatePropertyTypeIdentifier>());
 	PropertyModule.RegisterCustomPropertyTypeLayout(FCustomizableObjectVariation::StaticStruct()->GetFName(), FOnGetPropertyTypeCustomizationInstance::CreateStatic(&FCustomizableObjectVariationCustomization::MakeInstance));
 	PropertyModule.RegisterCustomPropertyTypeLayout(FCustomizableObjectTextureVariation::StaticStruct()->GetFName(), FOnGetPropertyTypeCustomizationInstance::CreateStatic(&FCustomizableObjectVariationCustomization::MakeInstance));
@@ -839,7 +839,7 @@ void FCustomizableObjectEditorModule::OnUpstreamCOsLoadedFixup(UCustomizableObje
 			Object->GetPrivate()->GetSource()->GetNodesOfClass<UCustomizableObjectNodeComponentMesh>(ComponentMeshNodes);
 			for (UCustomizableObjectNodeComponentMesh* ComponentNode : ComponentMeshNodes)
 			{
-				ComponentNode->LODSettings = RootObjectLODSettings;
+				ComponentNode->LODSettings_DEPRECATED = RootObjectLODSettings;
 			}
 		}
 	}
diff --git a/Engine/Plugins/Mutable/Source/CustomizableObjectEditor/Private/MuCOE/EdGraphSchema_CustomizableObject.cpp b/Engine/Plugins/Mutable/Source/CustomizableObjectEditor/Private/MuCOE/EdGraphSchema_CustomizableObject.cpp
--- a/Engine/Plugins/Mutable/Source/CustomizableObjectEditor/Private/MuCOE/EdGraphSchema_CustomizableObject.cpp
+++ b/Engine/Plugins/Mutable/Source/CustomizableObjectEditor/Private/MuCOE/EdGraphSchema_CustomizableObject.cpp
@@ -116,6 +116,7 @@
 #include "Nodes/CONodeExternalTypeParameter.h"
 #include "Nodes/CONodeRemoveMesh.h"
 #include "Nodes/CONodeSkeletalMeshBreak.h"
+#include "Nodes/CONodeMutableSkeletalMeshMake.h"
 #include "Nodes/CONodeTransformWithBone.h"
 #include "Nodes/CONodeTransformInMesh.h"
 #include "Nodes/CONodeSwitch.h"
@@ -645,6 +646,8 @@ void UEdGraphSchema_CustomizableObject::GetGraphContextActions(FGraphContextMenu
 			ContextMenuBuilder.CreateTemplateNode<UCustomizableObjectNodeSkeletalMeshParameter>(),
 			ContextMenuBuilder.CreateTemplateNode<UCustomizableObjectNodeMaterialParameter>(),
 			ContextMenuBuilder.CreateTemplateNode<UCONodeSkeletalMeshBreak>(),
+			ContextMenuBuilder.CreateTemplateNode<UCONodeMutableSkeletalMeshMake>(),
+			ContextMenuBuilder.CreateTemplateNode<UCONodeSkeletalMeshMake>(),
 		};
 
 		AddNewNodeCategoryActionsFiltered(ExperimentalTemplateNodes, ContextMenuBuilder, NC_Experimental, 2, Filter);
diff --git a/Engine/Plugins/Mutable/Source/CustomizableObjectEditor/Private/MuCOE/GenerateMutableSource/GenerateMutableSource.cpp b/Engine/Plugins/Mutable/Source/CustomizableObjectEditor/Private/MuCOE/GenerateMutableSource/GenerateMutableSource.cpp
--- a/Engine/Plugins/Mutable/Source/CustomizableObjectEditor/Private/MuCOE/GenerateMutableSource/GenerateMutableSource.cpp
+++ b/Engine/Plugins/Mutable/Source/CustomizableObjectEditor/Private/MuCOE/GenerateMutableSource/GenerateMutableSource.cpp
@@ -942,8 +942,15 @@ UE::Mutable::Private::Ptr<UE::Mutable::Private::NodeObject> GenerateMutableSourc
 				{
 					continue;
 				}
+
+				TArray<const UCustomizableObjectNodeMacroInstance*> ComponentMacroContext = ComponentInfo.MacroContext;
+				const UCONodeSkeletalMeshMake* ConnectedSkeletalMeshMakeNode = ComponentInfo.Node->GetConnectedSkeletalMeshMakeNode(
+					ComponentMacroContext);
+				
+				// Failure would mean that FirstPass did generate a componentInfo object for an invalid case 
+				check(ConnectedSkeletalMeshMakeNode);
 				
-				const FMutableLODSettings& LODSettings = ComponentInfo.Node->LODSettings;
+				const FMutableLODSettings& LODSettings = ConnectedSkeletalMeshMakeNode->LODSettings;
 				const ITargetPlatform* TargetPlatform = GenerationContext.CompilationContext->Options.TargetPlatform;
 				const FSupportedQualityLevelArray SupportedQualityLevels = LODSettings.MinQualityLevelLOD.GetSupportedQualityLevels(*TargetPlatform->GetPlatformInfo().IniPlatformName.ToString());
 				
diff --git a/Engine/Plugins/Mutable/Source/CustomizableObjectEditor/Private/MuCOE/GenerateMutableSource/GenerateMutableSourceComponent.cpp b/Engine/Plugins/Mutable/Source/CustomizableObjectEditor/Private/MuCOE/GenerateMutableSource/GenerateMutableSourceComponent.cpp
--- a/Engine/Plugins/Mutable/Source/CustomizableObjectEditor/Private/MuCOE/GenerateMutableSource/GenerateMutableSourceComponent.cpp
+++ b/Engine/Plugins/Mutable/Source/CustomizableObjectEditor/Private/MuCOE/GenerateMutableSource/GenerateMutableSourceComponent.cpp
@@ -98,15 +98,23 @@ UE::Mutable::Private::Ptr<UE::Mutable::Private::NodeComponent> GenerateMutableSo
 		FSkeletalMeshComponentInfo* ComponentInfo = GenerationContext.CompilationContext->GetComponentInfo(ComponentId);
 		if (ComponentInfo)
 		{
-			UCustomizableObjectNodeComponentMesh* TypedParentComponentMesh = ComponentInfo->Node;
+			const UCustomizableObjectNodeComponentMesh* TypedParentComponentMesh = ComponentInfo->Node;
 
-			if (TypedComponentMeshExtend->NumLODs > TypedParentComponentMesh->NumLODs)
+			TArray<const UCustomizableObjectNodeMacroInstance*> ComponentMacroContext = ComponentInfo->MacroContext;
+			const UCONodeSkeletalMeshMake* ConnectedMakeSkeletalMeshNode = TypedParentComponentMesh->GetConnectedSkeletalMeshMakeNode(
+				ComponentMacroContext);
+			
+			// Failure would mean that FirstPass did generate a componentInfo object for an invalid case 
+			// (not SKM connected from where to get the LOD count)
+			check(ConnectedMakeSkeletalMeshNode);
+			
+			if (TypedComponentMeshExtend->NumLODs > ConnectedMakeSkeletalMeshNode->NumLODs)
 			{
-				FText Msg = FText::Format(LOCTEXT("ExtendMeshComponentLODs", "Add To Mesh Component can not have more LODs than its parent Mesh Component [{0}]."), FText::FromName(ParentComponentName));
+				const FText Msg = FText::Format(LOCTEXT("ExtendMeshComponentLODs", "Add To Mesh Component can not have more LODs than its parent Mesh Component [{0}]."), FText::FromName(ParentComponentName));
 				GenerationContext.Log(Msg, TypedComponentMeshExtend, EMessageSeverity::Warning);
 			}
 
-			// Swap the macro contex since this component can be in another macro
+			// Swap the macro context since this component can be in another macro
 			TArray<const UCustomizableObjectNodeMacroInstance*> MacroContextCopy = GenerationContext.MacroNodesStack;
 			GenerationContext.MacroNodesStack = ComponentInfo->MacroContext;
 
@@ -372,39 +380,49 @@ UE::Mutable::Private::Ptr<UE::Mutable::Private::NodeComponent> GenerateMutableSo
 }
 
 
-void FirstPass(UCustomizableObjectNodeComponentMesh& Node, FMutableGraphGenerationContext& GenerationContext)
-{
-	USkeletalMesh* RefSkeletalMesh = Node.ReferenceSkeletalMesh;
+void FirstPass(UCustomizableObjectNodeComponentMesh& ComponentNode, FMutableGraphGenerationContext& GenerationContext)
+{	
+	USkeletalMesh* RefSkeletalMesh = ComponentNode.ReferenceSkeletalMesh;
 	if (!RefSkeletalMesh)
 	{
-		GenerationContext.Log(LOCTEXT("NoReferenceMeshObjectTab", "Error! Missing reference Skeletal Mesh"), &Node, EMessageSeverity::Error);
+		GenerationContext.Log(LOCTEXT("NoReferenceMeshObjectTab", "Missing reference Skeletal Mesh"), &ComponentNode, EMessageSeverity::Error);
 		return;
 	}
 	
 	USkeleton* RefSkeleton = RefSkeletalMesh->GetSkeleton();
 	if (!RefSkeleton)
 	{
-		FText Msg = FText::Format(LOCTEXT("NoReferenceSkeleton", "Error! Missing skeleton in the reference mesh [{0}]"), FText::FromString(GenerationContext.CustomizableObjectWithCycle->GetPathName()));
+		const FText Msg = FText::Format(LOCTEXT("NoReferenceSkeleton", "Missing skeleton in the reference mesh [{0}]"), FText::FromString(GenerationContext.CustomizableObjectWithCycle->GetPathName()));
+		GenerationContext.Log(Msg, &ComponentNode, EMessageSeverity::Error);
+		return;
+	}
+	
+	TArray<const UCustomizableObjectNodeMacroInstance*> LocalMacroContext = GenerationContext.MacroNodesStack;
 
-		GenerationContext.Log(Msg, &Node, EMessageSeverity::Error);
+	const UCONodeSkeletalMeshMake* ConnectedMakeSkeletalMeshNode = ComponentNode.GetConnectedSkeletalMeshMakeNode(LocalMacroContext);
+	if (!ConnectedMakeSkeletalMeshNode)
+	{
+		const FText Msg = FText::Format(LOCTEXT("NoSkeletalMeshPinConnected", "Missing connection to Skeletal Mesh Make node. [{0}]"), FText::FromString(GenerationContext.CustomizableObjectWithCycle->GetPathName()));
+		GenerationContext.Log(Msg, &ComponentNode, EMessageSeverity::Error);
 		return;
 	}
 	
 	// Fill the basic LOD Settings
-	uint8 NumLODs = Node.LODPins.Num();
-
+	uint8 NumLODs = ConnectedMakeSkeletalMeshNode->NumLODs;
+	check(NumLODs > 0);
+	
 	// NumLODsInRoot
-	int32 MaxRefMeshLODs = Node.ReferenceSkeletalMesh->GetLODNum();
+	int32 MaxRefMeshLODs = ComponentNode.ReferenceSkeletalMesh->GetLODNum();
 	if (MaxRefMeshLODs < NumLODs)
 	{
-		FString Msg = FString::Printf(TEXT("The object has %d LODs but the reference mesh only %d. Resulting objects will have %d LODs."),
+		const FString Msg = FString::Printf(TEXT("The object has %d LODs but the reference mesh only %d. Resulting objects will have %d LODs."),
 			NumLODs, MaxRefMeshLODs, MaxRefMeshLODs);
-		GenerationContext.Log(FText::FromString(Msg), &Node, EMessageSeverity::Warning);
+		GenerationContext.Log(FText::FromString(Msg), &ComponentNode, EMessageSeverity::Warning);
 		NumLODs = MaxRefMeshLODs;
 	}
 	
-	const FMutableLODSettings& LODSettings = Node.LODSettings;
-
+	const FMutableLODSettings& LODSettings = ConnectedMakeSkeletalMeshNode->LODSettings;
+	
 	uint8 FirstLODAvailable = MAX_MESH_LOD_COUNT;
 	
 	// Find the MinLOD available for the target platform
@@ -450,7 +468,7 @@ void FirstPass(UCustomizableObjectNodeComponentMesh& Node, FMutableGraphGenerati
 			NumMaxLODsToStream = LODSettings.NumMaxStreamedLODs.GetValueForPlatform(*GenerationContext.CompilationContext->Options.TargetPlatform->IniPlatformName());
 		}
 	}
-	else if(Node.ReferenceSkeletalMesh)
+	else if(ComponentNode.ReferenceSkeletalMesh)
 	{
 		if (RefSkeletalMesh->GetEnableLODStreaming(GenerationContext.CompilationContext->Options.TargetPlatform))
 		{
@@ -460,7 +478,7 @@ void FirstPass(UCustomizableObjectNodeComponentMesh& Node, FMutableGraphGenerati
 
 	NumMaxLODsToStream = FMath::Clamp(NumMaxLODsToStream, 0, NumLODs - 1);
 
-	const FName ComponentName = Node.GetComponentName(&GenerationContext.MacroNodesStack);
+	const FName ComponentName = ComponentNode.GetComponentName(&GenerationContext.MacroNodesStack);
 	UE::Mutable::Private::FComponentId ComponentId = GenerationContext.ComponentNames.Add(ComponentName);
 
 	GenerationContext.NumLODs.Add(ComponentId, NumLODs);
@@ -471,7 +489,7 @@ void FirstPass(UCustomizableObjectNodeComponentMesh& Node, FMutableGraphGenerati
 	FSkeletalMeshComponentInfo ComponentInfo;
 	ComponentInfo.ComponentName = ComponentName;
 	ComponentInfo.ComponentId = ComponentId;
-	ComponentInfo.Node = &Node;
+	ComponentInfo.Node = &ComponentNode;
 	ComponentInfo.MacroContext = GenerationContext.MacroNodesStack;
 	ComponentInfo.LODSettings = LODSettings;
 	
@@ -508,13 +526,14 @@ void FirstPass(UCustomizableObjectNodeComponentMesh& Node, FMutableGraphGenerati
 	ComponentInfo.BonesToRemovePerLOD.SetNum(NumLODs);
 
 	TMap<FName, bool> BonesToRemove;
-
-	const int32 ComponentSettingsLODCount = Node.LODReductionSettings.Num();
+	
+	const TArray<FLODReductionSettings>& LODReductionSettings = ConnectedMakeSkeletalMeshNode->LODReductionSettings;
+	const int32 ComponentSettingsLODCount = LODReductionSettings.Num();
 	for (int32 LODIndex = 0; LODIndex < NumLODs; ++LODIndex)
 	{
 		if (LODIndex < ComponentSettingsLODCount)
 		{
-			const FLODReductionSettings& LODReductionSetting = Node.LODReductionSettings[LODIndex];
+			const FLODReductionSettings& LODReductionSetting = LODReductionSettings[LODIndex];
 
 			for (const FBoneToRemove& Bone : LODReductionSetting.BonesToRemove)
 			{
diff --git a/Engine/Plugins/Mutable/Source/CustomizableObjectEditor/Private/MuCOE/GenerateMutableSource/GenerateMutableSourceComponent.h b/Engine/Plugins/Mutable/Source/CustomizableObjectEditor/Private/MuCOE/GenerateMutableSource/GenerateMutableSourceComponent.h
--- a/Engine/Plugins/Mutable/Source/CustomizableObjectEditor/Private/MuCOE/GenerateMutableSource/GenerateMutableSourceComponent.h
+++ b/Engine/Plugins/Mutable/Source/CustomizableObjectEditor/Private/MuCOE/GenerateMutableSource/GenerateMutableSourceComponent.h
@@ -17,5 +17,5 @@ namespace UE::Mutable::Private
 UE::Mutable::Private::Ptr<UE::Mutable::Private::NodeComponent> GenerateMutableSourceComponent(const UEdGraphPin* Pin, FMutableGraphGenerationContext& GenerationContext);
 
 
-void FirstPass(UCustomizableObjectNodeComponentMesh& Node, FMutableGraphGenerationContext& GenerationContext);
+void FirstPass(UCustomizableObjectNodeComponentMesh& ComponentNode, FMutableGraphGenerationContext& GenerationContext);
 
diff --git a/Engine/Plugins/Mutable/Source/CustomizableObjectEditor/Private/MuCOE/GenerateMutableSource/GenerateMutableSourceSkeletalMesh.cpp b/Engine/Plugins/Mutable/Source/CustomizableObjectEditor/Private/MuCOE/GenerateMutableSource/GenerateMutableSourceSkeletalMesh.cpp
--- a/Engine/Plugins/Mutable/Source/CustomizableObjectEditor/Private/MuCOE/GenerateMutableSource/GenerateMutableSourceSkeletalMesh.cpp
+++ b/Engine/Plugins/Mutable/Source/CustomizableObjectEditor/Private/MuCOE/GenerateMutableSource/GenerateMutableSourceSkeletalMesh.cpp
@@ -72,9 +72,28 @@ namespace UE::Mutable::Private
 		if (const ICustomizableObjectNodeComponentMeshInterface* TypedComponentMesh = Cast<ICustomizableObjectNodeComponentMeshInterface>(Node))
 		{
 			Ptr<NodeSkeletalMeshNew> SkeletalMeshNode = new NodeSkeletalMeshNew();
+			
+			TArray<FEdGraphPinReference> LODPins;
+			{
+				if (const UCustomizableObjectNodeComponentMesh* MeshComponentNode = Cast<UCustomizableObjectNodeComponentMesh>(Node))
+				{
+					const UCONodeMutableSkeletalMeshMake* ConnectedMutableSkeletalMeshMakeNode = MeshComponentNode->
+						GetConnectedMutableSkeletalMeshMakeNode(GenerationContext.MacroNodesStack);
+					
+					if (ConnectedMutableSkeletalMeshMakeNode)
+					{
+						LODPins = ConnectedMutableSkeletalMeshMakeNode->LODPins;
+					}
+				}
+				else
+				{
+					LODPins = TypedComponentMesh->GetLODPins();
+				}
+			}
+			
+			int32 FirstLOD = INDEX_NONE;
 
-			int32 FirstLOD = -1;
-
+			const int32 NumLODPins = LODPins.Num();
 			const int32 NumLODsInRoot = GenerationContext.NumLODs[GenerationContext.CurrentSkeletalMeshComponent];
 			for (int32 CurrentLOD = 0; CurrentLOD < NumLODsInRoot; ++CurrentLOD)
 			{
@@ -88,10 +107,8 @@ namespace UE::Mutable::Private
 				Ptr<NodeLOD> LODNode = SkeletalMeshNode->LODs[CurrentLOD];
 
 				LODNode->SetMessageContext(&TypedComponentMesh);
-
-				const int32 NumLODs = TypedComponentMesh->GetLODPins().Num();
-
-				FirstLOD = CurrentLOD < NumLODs && FirstLOD == INDEX_NONE ? CurrentLOD : FirstLOD;
+				
+				FirstLOD = CurrentLOD < NumLODPins && FirstLOD == INDEX_NONE ? CurrentLOD : FirstLOD;
 
 				if (FirstLOD < 0)
 				{
@@ -106,12 +123,12 @@ namespace UE::Mutable::Private
 				// Generate all relevant LODs for this object up until the current LODIndex.
 				for (int32 LODIndex = FirstLOD; LODIndex <= CurrentLOD; ++LODIndex)
 				{
-					if (!TypedComponentMesh->GetLODPins().IsValidIndex(LODIndex))
+					if (!LODPins.IsValidIndex(LODIndex))
 					{
 						continue;
 					}
 					
-					const UEdGraphPin* LODPin = TypedComponentMesh->GetLODPins()[LODIndex].Get();
+					const UEdGraphPin* LODPin = LODPins[LODIndex].Get();
 					check(LODPin);
 					
 					GenerationContext.FromLOD = LODIndex;
diff --git a/Engine/Plugins/Mutable/Source/CustomizableObjectEditor/Private/MuCOE/Nodes/CONodeMutableSkeletalMeshMake.cpp b/Engine/Plugins/Mutable/Source/CustomizableObjectEditor/Private/MuCOE/Nodes/CONodeMutableSkeletalMeshMake.cpp
new file mode 100644
--- /dev/null
+++ b/Engine/Plugins/Mutable/Source/CustomizableObjectEditor/Private/MuCOE/Nodes/CONodeMutableSkeletalMeshMake.cpp
@@ -0,0 +1,100 @@
+// Copyright Epic Games, Inc. All Rights Reserved.
+
+#include "MuCOE/Nodes/CONodeMutableSkeletalMeshMake.h"
+
+#include "MuCOE/EdGraphSchema_CustomizableObject.h"
+#include "ScopedTransaction.h"
+#include "MuCOE/RemapPins/CustomizableObjectNodeRemapPinsByPosition.h"
+
+#include UE_INLINE_GENERATED_CPP_BY_NAME(CONodeMutableSkeletalMeshMake)
+
+
+#define LOCTEXT_NAMESPACE "CustomizableObjectEditor"
+
+
+void UCONodeMutableSkeletalMeshMake::AllocateDefaultPins(UCustomizableObjectNodeRemapPins* RemapPins)
+{
+	// Inputs
+	const int32 NumLODs = !LODPins.IsEmpty() ? LODPins.Num() : 1;
+	LODPins.Empty(NumLODs);
+	for (uint16 LODIndex = 0; LODIndex < NumLODs; ++LODIndex)
+	{
+		const FString LODName = GetNameForLODPin(LODIndex);
+		UEdGraphPin* LODPin = CustomCreatePin(EGPD_Input, UEdGraphSchema_CustomizableObject::PC_MeshSection, FName(*LODName), FText::FromString(LODName));
+		LODPin->PinType.ContainerType = EPinContainerType::Array;
+
+		LODPins.Add(LODPin);
+	}
+	
+	// Output
+	SkeletalMeshPin = CustomCreatePin(EGPD_Output, UEdGraphSchema_CustomizableObject::PC_SkeletalMesh);
+}
+
+
+UCustomizableObjectNodeRemapPins* UCONodeMutableSkeletalMeshMake::CreateRemapPinsDefault() const
+{
+	return CreateRemapPinsByPosition();
+}
+
+
+EAddPinNodeButtonLocation UCONodeMutableSkeletalMeshMake::GetAddPinButtonNodeSide() const
+{
+	return EAddPinNodeButtonLocation::INPUT;
+}
+
+
+void UCONodeMutableSkeletalMeshMake::AddPinFromUI()
+{
+	Super::AddPinFromUI();
+
+	FScopedTransaction Transaction(LOCTEXT("CONodeMutableSkeletalMeshMake_AddPinFromNodeUITransaction", "Add Pin from Node UI"));
+	Modify();
+	
+	const FString LODName = GetNameForLODPin(LODPins.Num());
+	UEdGraphPin* LODPin = CustomCreatePin(EGPD_Input, UEdGraphSchema_CustomizableObject::PC_MeshSection, FName(*LODName), FText::FromString(LODName));
+	LODPin->PinType.ContainerType = EPinContainerType::Array;
+
+	LODPins.Add(LODPin);
+	
+	GetGraph()->NotifyNodeChanged(this);
+}
+
+
+bool UCONodeMutableSkeletalMeshMake::CanPinBeRemoved(const UEdGraphPin& Pin) const
+{
+	// Disallow the removal of already connected pins (to prevent human error)
+	// Only allow the removal of a pin if that pin is one definning a lod there are more than one LODs available (0 LODs is not a valid option)
+	return Pin.Direction == EGPD_Input && LODPins.Num() > 1 && !Pin.LinkedTo.Num() && !Pin.bOrphanedPin  && LODPins.Contains(&Pin);
+}
+
+
+bool UCONodeMutableSkeletalMeshMake::CustomRemovePin(UEdGraphPin& Pin)
+{
+	if (LODPins.Remove(&Pin))
+	{
+		const bool bWasRemovalSuccessful = Super::CustomRemovePin(Pin);
+		check(bWasRemovalSuccessful);
+	
+		ReconstructNode();
+	
+		return bWasRemovalSuccessful;
+	}
+	else
+	{
+		return Super::CustomRemovePin(Pin);
+	}
+}
+
+
+FText UCONodeMutableSkeletalMeshMake::GetNodeTitle(ENodeTitleType::Type TitleType) const
+{
+	return LOCTEXT("Mutable_Skeletal_Mesh_Make_Node_Title", "Mutable Skeletal Mesh Make");
+}
+
+
+FLinearColor UCONodeMutableSkeletalMeshMake::GetNodeTitleColor() const
+{
+	return UEdGraphSchema_CustomizableObject::GetPinTypeColor(UEdGraphSchema_CustomizableObject::PC_SkeletalMesh_Passthrough);
+}
+
+#undef LOCTEXT_NAMESPACE
diff --git a/Engine/Plugins/Mutable/Source/CustomizableObjectEditor/Private/MuCOE/Nodes/CONodeMutableSkeletalMeshMake.h b/Engine/Plugins/Mutable/Source/CustomizableObjectEditor/Private/MuCOE/Nodes/CONodeMutableSkeletalMeshMake.h
new file mode 100644
--- /dev/null
+++ b/Engine/Plugins/Mutable/Source/CustomizableObjectEditor/Private/MuCOE/Nodes/CONodeMutableSkeletalMeshMake.h
@@ -0,0 +1,44 @@
+// Copyright Epic Games, Inc. All Rights Reserved.
+
+#pragma once
+
+#include "MuCOE/Nodes/CustomizableObjectNodeMeshInterface.h"
+#include "MuCOE/Nodes/SCustomizableObjectNode.h"
+
+#include "CONodeMutableSkeletalMeshMake.generated.h"
+
+
+
+UCLASS()
+class UCONodeMutableSkeletalMeshMake : public UCustomizableObjectNode
+{
+public:
+	GENERATED_BODY()
+	
+	// UEdGraphNode
+	virtual FText GetNodeTitle(ENodeTitleType::Type TitleType) const override;
+	virtual FLinearColor GetNodeTitleColor() const override;
+	
+	// UCustomizableObjectNode
+	virtual void AllocateDefaultPins(UCustomizableObjectNodeRemapPins* RemapPins) override;
+	virtual UCustomizableObjectNodeRemapPins* CreateRemapPinsDefault() const override;
+	virtual EAddPinNodeButtonLocation GetAddPinButtonNodeSide() const override;
+	virtual void AddPinFromUI() override;
+	virtual bool CanPinBeRemoved(const UEdGraphPin& Pin) const override;
+	virtual bool CustomRemovePin(UEdGraphPin& Pin) override;
+	
+private:
+	
+	FString GetNameForLODPin(const int32& InIndex) const
+	{
+		return FString::Printf(TEXT("LOD %d"), InIndex);
+	}
+	
+public:
+	
+	UPROPERTY()
+	TArray<FEdGraphPinReference> LODPins;
+	
+	UPROPERTY()
+	FEdGraphPinReference SkeletalMeshPin;
+};
diff --git a/Engine/Plugins/Mutable/Source/CustomizableObjectEditor/Private/MuCOE/Nodes/CONodeSkeletalMeshMake.cpp b/Engine/Plugins/Mutable/Source/CustomizableObjectEditor/Private/MuCOE/Nodes/CONodeSkeletalMeshMake.cpp
new file mode 100644
--- /dev/null
+++ b/Engine/Plugins/Mutable/Source/CustomizableObjectEditor/Private/MuCOE/Nodes/CONodeSkeletalMeshMake.cpp
@@ -0,0 +1,133 @@
+// Copyright Epic Games, Inc. All Rights Reserved.
+
+#include "MuCOE/Nodes/CONodeSkeletalMeshMake.h"
+
+#include "CONodeMutableSkeletalMeshMake.h"
+#include "CustomizableObjectNodeStaticString.h"
+#include "MuCOE/EdGraphSchema_CustomizableObject.h"
+#include "MuCOE/GraphTraversal.h"
+
+#include "ScopedTransaction.h"
+
+#include UE_INLINE_GENERATED_CPP_BY_NAME(CONodeSkeletalMeshMake)
+
+#define LOCTEXT_NAMESPACE "CustomizableObjectEditor"
+
+UCONodeSkeletalMeshMake::UCONodeSkeletalMeshMake()
+{
+	const FString CVarName = TEXT("r.SkeletalMesh.MinLodQualityLevel");
+	const FString ScalabilitySectionName = TEXT("ViewDistanceQuality");
+	LODSettings.MinQualityLevelLOD.SetQualityLevelCVarForCooking(*CVarName, *ScalabilitySectionName);
+}
+
+
+FText UCONodeSkeletalMeshMake::GetNodeTitle(ENodeTitleType::Type TitleType) const
+{
+	return LOCTEXT("Skeletal_Mesh_Make_Node_Title", "Skeletal Mesh Make");
+}
+
+
+FLinearColor UCONodeSkeletalMeshMake::GetNodeTitleColor() const
+{
+	return UEdGraphSchema_CustomizableObject::GetPinTypeColor(UEdGraphSchema_CustomizableObject::PC_SkeletalMesh);
+}
+
+
+void UCONodeSkeletalMeshMake::AllocateDefaultPins(UCustomizableObjectNodeRemapPins* RemapPins)
+{
+	MeshNamePin = CustomCreatePin(EEdGraphPinDirection::EGPD_Input, UEdGraphSchema_CustomizableObject::PC_String);
+	SkeletalMeshPin = CustomCreatePin(EEdGraphPinDirection::EGPD_Input, UEdGraphSchema_CustomizableObject::PC_SkeletalMesh);
+	
+	LODReductionSettings.SetNum(NumLODs);
+	
+	PassthroughSkeletalMeshPin = CustomCreatePin(EEdGraphPinDirection::EGPD_Output, UEdGraphSchema_CustomizableObject::PC_SkeletalMesh_Passthrough);
+}
+
+
+bool UCONodeSkeletalMeshMake::CanRenamePin(const UEdGraphPin& Pin) const
+{
+	return Pin.PinType.PinCategory == UEdGraphSchema_CustomizableObject::PC_String;
+}
+
+
+FText UCONodeSkeletalMeshMake::GetPinEditableName(const UEdGraphPin& Pin) const
+{
+	const UEdGraphPin* ComponentPin = MeshNamePin.Get();
+	if (ComponentPin && ComponentPin->PinId == Pin.PinId)
+	{
+		return FText::FromString(GetSkeletalMeshName());
+	}
+
+	return FText::GetEmpty();
+}
+
+
+void UCONodeSkeletalMeshMake::SetPinEditableName(const UEdGraphPin& Pin, const FText& InValue)
+{
+	const UEdGraphPin* ComponentPin = MeshNamePin.Get();
+	if (ComponentPin && ComponentPin->PinId == Pin.PinId)
+	{
+		FScopedTransaction LocalTransaction(LOCTEXT("UCONodeSkeletalMeshMake_SetSKMPinEditableNameTransaction", "Change SKMName PinName"));
+		Modify();
+		
+		SkeletalMeshName = InValue.ToString();
+		PinEditableNameChangedDelegate.Broadcast(Pin, InValue);
+	}
+}
+
+
+void UCONodeSkeletalMeshMake::PostEditChangeProperty(FPropertyChangedEvent& PropertyChangedEvent)
+{
+	Super::PostEditChangeProperty(PropertyChangedEvent);
+	
+	const FProperty* PropertyThatChanged = PropertyChangedEvent.Property;
+	if (!PropertyThatChanged)
+	{
+		return;
+	}
+
+	if (PropertyThatChanged->GetFName() == GET_MEMBER_NAME_CHECKED(UCONodeSkeletalMeshMake, NumLODs))
+	{
+		LODReductionSettings.SetNum(NumLODs);
+	}
+}
+
+
+const UCONodeMutableSkeletalMeshMake* UCONodeSkeletalMeshMake::GetConnectedMutableSkeletalMeshMakeNode(
+	TArray<const UCustomizableObjectNodeMacroInstance*>& MacroContext) const
+{
+	if (const UEdGraphPin* MeshPin = SkeletalMeshPin.Get())
+	{
+		if (const UEdGraphPin* LinkedPin = FollowInputPin(*MeshPin))
+		{
+			const UEdGraphPin* SkeletalMeshMakeOutputPin = GraphTraversal::FindIOPinSourceThroughMacroContext(*LinkedPin, &MacroContext);
+
+			if (const UCONodeMutableSkeletalMeshMake* MutableMakeNode = SkeletalMeshMakeOutputPin ? Cast<UCONodeMutableSkeletalMeshMake>(SkeletalMeshMakeOutputPin->GetOwningNode()) : nullptr)
+			{
+				return MutableMakeNode;
+			}
+		}
+	}
+	
+	return nullptr;
+}
+
+
+FString UCONodeSkeletalMeshMake::GetSkeletalMeshName(TArray<const UCustomizableObjectNodeMacroInstance*>* MacroContext) const
+{
+	if (const UEdGraphPin* NamePin = MeshNamePin.Get())
+	{
+		if (const UEdGraphPin* LinkedPin = FollowInputPin(*NamePin))
+		{
+			const UEdGraphPin* StringPin = GraphTraversal::FindIOPinSourceThroughMacroContext(*LinkedPin, MacroContext);
+			if (const UCustomizableObjectNodeStaticString* StringNode = StringPin ? Cast<UCustomizableObjectNodeStaticString>(StringPin->GetOwningNode()) : nullptr)
+			{
+				return StringNode->Value;
+			}
+		}
+	}
+
+	return SkeletalMeshName;
+}
+
+#undef LOCTEXT_NAMESPACE
diff --git a/Engine/Plugins/Mutable/Source/CustomizableObjectEditor/Private/MuCOE/Nodes/CONodeSkeletalMeshMake.h b/Engine/Plugins/Mutable/Source/CustomizableObjectEditor/Private/MuCOE/Nodes/CONodeSkeletalMeshMake.h
new file mode 100644
--- /dev/null
+++ b/Engine/Plugins/Mutable/Source/CustomizableObjectEditor/Private/MuCOE/Nodes/CONodeSkeletalMeshMake.h
@@ -0,0 +1,89 @@
+// Copyright Epic Games, Inc. All Rights Reserved.
+
+#pragma once
+
+#include "MuCOE/Nodes/CustomizableObjectNodeMeshInterface.h"
+#include "MuCOE/Nodes/SCustomizableObjectNode.h"
+
+#include "CONodeSkeletalMeshMake.generated.h"
+
+class UCONodeMutableSkeletalMeshMake;
+
+USTRUCT()
+struct FBoneToRemove
+{
+	GENERATED_USTRUCT_BODY()
+
+	UPROPERTY(EditAnywhere, Category = SkeletalMesh)
+	FName BoneName;
+	
+	UPROPERTY(EditAnywhere, Category = SkeletalMesh)
+	bool bOnlyRemoveChildren = false;
+};
+
+
+USTRUCT()
+struct FLODReductionSettings
+{
+	GENERATED_USTRUCT_BODY()
+
+	/** Selects which bones will be removed from the final skeleton
+	* BoneName: Name of the bone that will be removed. Its children will be removed too.
+	* Remove Only Children: If true, only the children of the selected bone will be removed. The selected bone will remain.
+	*/
+	UPROPERTY(EditAnywhere, Category = SkeletalMesh, meta=(TitleProperty="BoneName"))	
+	TArray<FBoneToRemove> BonesToRemove;
+};
+
+
+UCLASS()
+class UCONodeSkeletalMeshMake : public UCustomizableObjectNode
+{
+public:
+	GENERATED_BODY()
+	
+	UCONodeSkeletalMeshMake();
+	
+	// UEdGraphNode
+	virtual FText GetNodeTitle(ENodeTitleType::Type TitleType) const override;
+	virtual FLinearColor GetNodeTitleColor() const override;
+	
+	// UCustomizableObjectNode
+	virtual void AllocateDefaultPins(UCustomizableObjectNodeRemapPins* RemapPins) override;
+	virtual bool CanRenamePin(const UEdGraphPin& Pin) const override;
+	virtual FText GetPinEditableName(const UEdGraphPin& Pin) const override;
+	virtual void SetPinEditableName(const UEdGraphPin& Pin, const FText& InValue) override;
+	virtual void PostEditChangeProperty(FPropertyChangedEvent& PropertyChangedEvent) override;
+
+	/**
+	 * Get a pointer to the connected UCONodeMutableSkeletalMeshMake node.
+	 * @param MacroContext The current macro stack
+	 * @return A valid pointer if the node is connected, a null pointer otherwise.
+	 */
+	const UCONodeMutableSkeletalMeshMake* GetConnectedMutableSkeletalMeshMakeNode(TArray<const UCustomizableObjectNodeMacroInstance*>& MacroContext) const;;
+	
+	FString GetSkeletalMeshName(TArray<const UCustomizableObjectNodeMacroInstance*>* MacroContext = nullptr) const;
+	
+	UPROPERTY(EditAnywhere, Category=SkeletalMesh, meta = (ClampMin=1))
+	uint16 NumLODs = 1;
+	
+	UPROPERTY(EditAnywhere, Category = SkeletalMesh)
+	FMutableLODSettings LODSettings;
+	
+	UPROPERTY(EditAnywhere, EditFixedSize, Category=SkeletalMesh)
+	TArray<FLODReductionSettings> LODReductionSettings;
+	
+	/** Name to use by the generated skeletal mesh object. */
+	UPROPERTY()
+	FEdGraphPinReference MeshNamePin;
+	
+	UPROPERTY()
+	FEdGraphPinReference SkeletalMeshPin;
+	
+	UPROPERTY()
+	FEdGraphPinReference PassthroughSkeletalMeshPin;
+	
+private:
+	UPROPERTY()
+	FString SkeletalMeshName = "Default Name";
+};
\ No newline at end of file
diff --git a/Engine/Plugins/Mutable/Source/CustomizableObjectEditor/Private/MuCOE/Nodes/CustomizableObjectNodeComponentMesh.cpp b/Engine/Plugins/Mutable/Source/CustomizableObjectEditor/Private/MuCOE/Nodes/CustomizableObjectNodeComponentMesh.cpp
--- a/Engine/Plugins/Mutable/Source/CustomizableObjectEditor/Private/MuCOE/Nodes/CustomizableObjectNodeComponentMesh.cpp
+++ b/Engine/Plugins/Mutable/Source/CustomizableObjectEditor/Private/MuCOE/Nodes/CustomizableObjectNodeComponentMesh.cpp
@@ -3,6 +3,7 @@
 #include "MuCOE/Nodes/CustomizableObjectNodeComponentMesh.h"
 
 #include "CONodeMaterialConstant.h"
+#include "CONodeSkeletalMeshMake.h"
 #include "MuCO/CustomizableObjectCustomVersion.h"
 #include "MuCOE/CustomizableObjectSchemaActions.h"
 #include "MuCOE/EdGraphSchema_CustomizableObject.h"
@@ -14,30 +15,12 @@
 #define LOCTEXT_NAMESPACE "CustomizableObjectEditor"
 
 
+
 UCustomizableObjectNodeComponentMesh::UCustomizableObjectNodeComponentMesh()
 {
 	const FString CVarName = TEXT("r.SkeletalMesh.MinLodQualityLevel");
 	const FString ScalabilitySectionName = TEXT("ViewDistanceQuality");
-	LODSettings.MinQualityLevelLOD.SetQualityLevelCVarForCooking(*CVarName, *ScalabilitySectionName);
-}
-
-
-void UCustomizableObjectNodeComponentMesh::PostEditChangeProperty(FPropertyChangedEvent& PropertyChangedEvent)
-{
-	Super::PostEditChangeProperty(PropertyChangedEvent);
-
-	const FProperty* PropertyThatChanged = PropertyChangedEvent.Property;
-	if (!PropertyThatChanged)
-	{
-		return;
-	}
-
-	if (PropertyThatChanged->GetFName() == GET_MEMBER_NAME_CHECKED(UCustomizableObjectNodeComponentMesh, NumLODs))
-	{
-		LODReductionSettings.SetNum(NumLODs);
-		
-		ReconstructNode();
-	}
+	LODSettings_DEPRECATED.MinQualityLevelLOD.SetQualityLevelCVarForCooking(*CVarName, *ScalabilitySectionName);
 }
 
 
@@ -96,11 +79,121 @@ void UCustomizableObjectNodeComponentMesh::BackwardsCompatibleFixup(int32 Custom
 			UEdGraphPin* MaterialConstantNodeMaterialPin = MaterialConstantNode->MaterialPin.Get();
 			check(MaterialConstantNodeMaterialPin);
 			
-			MaterialConstantNodeMaterialPin->MakeLinkTo(OverlayMaterialPinPtr);
+			MaterialConstantNodeMaterialPin->MakeLinkTo(OverlayMaterialPinPtr, false);
 			Graph->NotifyNodeChanged(MaterialConstantNode);
 			Graph->NotifyNodeChanged(this);
 		}
 	}
+	else if (CustomizableObjectCustomVersion == FCustomizableObjectCustomVersion::SkeletalMeshMakeNodeDeprecations)
+	{
+		// This node LOD arrays no longer have a use for us as those are defined by the connected node to the Skeletal mesh pin
+		// That pin does not yet exist
+
+		// First of all ensure we have the new Skeletal Mesh Input pin by creating it
+		check(!SkeletalMeshPin.Get());
+		SkeletalMeshPin = CustomCreatePin(EGPD_Input, UEdGraphSchema_CustomizableObject::PC_SkeletalMesh_Passthrough);
+		
+		// Spawn a new CONodeSkeletalMeshMake
+		UEdGraph* Graph = GetGraph();
+		check(Graph);
+		
+		constexpr float Separation = 400;
+		const FVector2D DefaultSpawnPosition = FVector2D(GetNodePosX(), GetNodePosY());
+		
+		// Intermediate node between the Mutable Skm Make node and the component node (this)
+		// Will always exist as now it is the node that holds the data that the component node was holding until now
+		UCONodeSkeletalMeshMake* MakeMeshNode = nullptr;
+		{
+			MakeMeshNode = Cast<UCONodeSkeletalMeshMake>(FCustomizableObjectSchemaAction_NewNode::CreateNode(Graph, nullptr, DefaultSpawnPosition, Cast<UEdGraphNode>(NewObject<UCONodeSkeletalMeshMake>())));
+			MakeMeshNode->LODSettings = LODSettings_DEPRECATED;
+			MakeMeshNode->LODReductionSettings = LODReductionSettings_DEPRECATED;
+			MakeMeshNode->NumLODs = NumLODs_DEPRECATED;
+				
+			// It is possible for the LODReductionSettings_DEPRECATED array to be empty if the NumLODs_DEPRECATED was left at the default value of 1.
+			// In that case, due to a bug, the LODReductionSettings_DEPRECATED would be left empty until the NumLODs_DEPRECATED was changed by the user
+			if (MakeMeshNode->LODReductionSettings.Num() != MakeMeshNode->NumLODs)
+			{
+				MakeMeshNode->LODReductionSettings.SetNum(MakeMeshNode->NumLODs);
+			}
+			
+			MakeMeshNode->PassthroughSkeletalMeshPin.Get()->MakeLinkTo(SkeletalMeshPin.Get(), false);
+		}
+		
+		// Move this to a side and leave the new node the position it has
+		SetNodePosX( DefaultSpawnPosition.X + MakeMeshNode->GetWidth() + Separation);
+		
+		// First array is the LODIndex and the second one the pins connected to said LOD Index
+		bool bIsAtLeastOneLODPinConnected = false;
+		for (FEdGraphPinReference& LODArrayPin : LODPins_DEPRECATED)
+		{
+			check(LODArrayPin.Get());
+			
+			// The LOD Pin is connected to a series of sections. Keep them as they will later be connected to the new node (CONOdeSkeletalMeshMake)
+			const TArray<UEdGraphPin*> LinkedToPins = LODArrayPin.Get()->LinkedTo;
+			if (!LinkedToPins.IsEmpty())
+			{
+		 		bIsAtLeastOneLODPinConnected = true;
+		 		break;
+			}
+		}
+		
+		// If at leat one of the LODs do have a connected pin then generate the mutable make node
+		if (bIsAtLeastOneLODPinConnected)
+		{
+			// Set node to the left of this node
+			// const FVector2D MakeMutableMeshNodePos = FVector2D(DefaultSpawnPosition.X - 400, GetNodePosY());
+			UCONodeMutableSkeletalMeshMake* MakeMutableMeshNode = Cast<UCONodeMutableSkeletalMeshMake>(FCustomizableObjectSchemaAction_NewNode::CreateNode(Graph, nullptr, DefaultSpawnPosition, Cast<UEdGraphNode>(NewObject<UCONodeMutableSkeletalMeshMake>())));
+			
+			// Connect both nodes
+			MakeMutableMeshNode->SkeletalMeshPin.Get()->MakeLinkTo(MakeMeshNode->SkeletalMeshPin.Get(), false);
+			
+			// Connect to this new node all the connections set in the current LODArrayPin and remove them from the current pin
+			const uint16 DefaultAmountOfLodPins = MakeMutableMeshNode->LODPins.Num(); 
+			for (uint16 LODIndex = DefaultAmountOfLodPins; LODIndex < NumLODs_DEPRECATED; ++LODIndex)
+			{
+				const FString LODName = FString::Printf(TEXT("LOD %d"), LODIndex);
+				
+				UEdGraphPin* LODPin = MakeMutableMeshNode->CustomCreatePin(EGPD_Input, UEdGraphSchema_CustomizableObject::PC_MeshSection, FName(*LODName), FText::FromString(LODName));
+				LODPin->PinType.ContainerType = EPinContainerType::Array;
+				
+				MakeMutableMeshNode->LODPins.Add(LODPin);
+			}
+
+			check(MakeMutableMeshNode->LODPins.Num() == LODPins_DEPRECATED.Num());
+			
+			// At this point all lod pins have been generated so it should be ok to link them with the old data from this node
+			for (uint16 LODIndex = 0; LODIndex < NumLODs_DEPRECATED; ++LODIndex)
+			{
+				UEdGraphPin* MakeNodeLODPin = MakeMutableMeshNode->LODPins[LODIndex].Get();
+				UEdGraphPin* ThisNodeLODPin = LODPins_DEPRECATED[LODIndex].Get();
+
+				// Se the sections for the make node and then reset the said connections from this node
+				for (UEdGraphPin* LinkedToPin : ThisNodeLODPin->LinkedTo)
+				{
+					MakeNodeLODPin->MakeLinkTo(LinkedToPin, false);
+				}
+				
+				ThisNodeLODPin->BreakAllPinLinks(false, false);
+			}
+
+			Graph->NotifyNodeChanged(MakeMutableMeshNode);
+			
+			// Move this to a side and leave the new node the position it has
+			SetNodePosX(DefaultSpawnPosition.X + Separation * 2);
+			MakeMeshNode->SetNodePosX(DefaultSpawnPosition.X + Separation);
+		}
+		
+		// Now remove the deprecated LODPins from this node
+		for (FEdGraphPinReference LODPinRef : LODPins_DEPRECATED)
+		{
+			UEdGraphPin* LodPin = LODPinRef.Get();
+			check(LodPin);
+			CustomRemovePin(*LodPin);
+		}
+		
+		Graph->NotifyNodeChanged(MakeMeshNode);
+		Graph->NotifyNodeChanged(this);
+	}
 }
 
 
@@ -112,17 +205,7 @@ void UCustomizableObjectNodeComponentMesh::AllocateDefaultPins(UCustomizableObje
 	OverlayMaterialPin = CustomCreatePin(EGPD_Input, UEdGraphSchema_CustomizableObject::PC_Material, TEXT("Overlay Material_Input_Pin"), LOCTEXT("OverlayMaterial", "Overlay Material"));
 	OverlayMaterialPin.Get()->PinToolTip = "Pin for an Overlay Material from a Table Node"; 
 	
-	// Base Mesh Interface Pins
-	LODPins.Empty(NumLODs);
-	for (int32 LODIndex = 0; LODIndex < NumLODs; ++LODIndex)
-	{
-		FString LODName = FString::Printf(TEXT("LOD %d"), LODIndex);
-
-		UEdGraphPin* Pin = CustomCreatePin(EGPD_Input, UEdGraphSchema_CustomizableObject::PC_MeshSection, FName(*LODName), FText::FromString(LODName));
-		Pin->PinType.ContainerType = EPinContainerType::Array;
-
-		LODPins.Add(Pin);
-	}
+	SkeletalMeshPin = CustomCreatePin(EGPD_Input, UEdGraphSchema_CustomizableObject::PC_SkeletalMesh_Passthrough);
 }
 
 
@@ -132,8 +215,6 @@ FText UCustomizableObjectNodeComponentMesh::GetNodeTitle(ENodeTitleType::Type Ti
 }
 
 
-
-
 FEdGraphPinReference UCustomizableObjectNodeComponentMesh::GetOverlayMaterialAssetPin() const
 {
 	return OverlayMaterialPin.Get();
@@ -148,15 +229,34 @@ bool UCustomizableObjectNodeComponentMesh::IsSingleOutputNode() const
 }
 
 
+bool UCustomizableObjectNodeComponentMesh::CanConnect(const UEdGraphPin* InOwnedInputPin, const UEdGraphPin* InOutputPin,
+	bool& bOutIsOtherNodeBlocklisted, bool& bOutArePinsCompatible) const
+{
+	if (!Super::CanConnect(InOwnedInputPin, InOutputPin, bOutIsOtherNodeBlocklisted, bOutArePinsCompatible))
+	{
+		return false;
+	}
+
+	if (InOwnedInputPin == SkeletalMeshPin.Get())
+	{
+		const UEdGraphNode* OutputPinOwningNode = InOutputPin->GetOwningNode();
+		bOutIsOtherNodeBlocklisted = !OutputPinOwningNode->IsA(UCONodeSkeletalMeshMake::StaticClass());		// prevent connections to any kind of node that it is not the make one
+		return !bOutIsOtherNodeBlocklisted;
+	}
+
+	return false;
+}
+
+
 int32 UCustomizableObjectNodeComponentMesh::GetNumLODs()
 {
-	return NumLODs;
+	return NumLODs_DEPRECATED;
 }
 
 
 const TArray<FEdGraphPinReference>& UCustomizableObjectNodeComponentMesh::GetLODPins() const
 {
-	return LODPins;
+	return LODPins_DEPRECATED;
 }
 
 
@@ -172,10 +272,40 @@ void UCustomizableObjectNodeComponentMesh::SetOutputPin(const UEdGraphPin* Pin)
 }
 
 
- const UCustomizableObjectNode* UCustomizableObjectNodeComponentMesh::GetOwningNode() const
+const UCustomizableObjectNode* UCustomizableObjectNodeComponentMesh::GetOwningNode() const
 {
 	 return this;
 }
 
-#undef LOCTEXT_NAMESPACE
 
+const UCONodeSkeletalMeshMake* UCustomizableObjectNodeComponentMesh::GetConnectedSkeletalMeshMakeNode(TArray<const UCustomizableObjectNodeMacroInstance*>& MacroContext) const
+{
+	if (const UEdGraphPin* MeshPin = SkeletalMeshPin.Get())
+	{
+		if (const UEdGraphPin* LinkedPin = FollowInputPin(*MeshPin))
+		{
+			const UEdGraphPin* SkeletalMeshMakeOutputPin = GraphTraversal::FindIOPinSourceThroughMacroContext(*LinkedPin, &MacroContext);
+
+			if (const UCONodeSkeletalMeshMake* MakeNode = SkeletalMeshMakeOutputPin ? Cast<UCONodeSkeletalMeshMake>(SkeletalMeshMakeOutputPin->GetOwningNode()) : nullptr)
+			{
+				return MakeNode;
+			}
+		}
+	}
+	
+	return nullptr;
+}
+
+
+const UCONodeMutableSkeletalMeshMake* UCustomizableObjectNodeComponentMesh::GetConnectedMutableSkeletalMeshMakeNode(TArray<const UCustomizableObjectNodeMacroInstance*>& MacroContext) const
+{
+	if (const UCONodeSkeletalMeshMake* MakeSkeletalMeshNode = GetConnectedSkeletalMeshMakeNode(MacroContext))
+	{
+		return MakeSkeletalMeshNode->GetConnectedMutableSkeletalMeshMakeNode(MacroContext);
+	}
+	
+	return nullptr;
+}
+
+
+#undef LOCTEXT_NAMESPACE
diff --git a/Engine/Plugins/Mutable/Source/CustomizableObjectEditor/Private/MuCOE/Nodes/CustomizableObjectNodeComponentMesh.h b/Engine/Plugins/Mutable/Source/CustomizableObjectEditor/Private/MuCOE/Nodes/CustomizableObjectNodeComponentMesh.h
--- a/Engine/Plugins/Mutable/Source/CustomizableObjectEditor/Private/MuCOE/Nodes/CustomizableObjectNodeComponentMesh.h
+++ b/Engine/Plugins/Mutable/Source/CustomizableObjectEditor/Private/MuCOE/Nodes/CustomizableObjectNodeComponentMesh.h
@@ -2,14 +2,14 @@
 
 #pragma once
 
+#include "CONodeMutableSkeletalMeshMake.h"
+#include "CONodeSkeletalMeshMake.h"
 #include "CustomizableObjectNodeComponent.h"
 #include "CustomizableObjectNodeComponentMeshBase.h"
 #include "MuR/System.h"
 
 #include "CustomizableObjectNodeComponentMesh.generated.h"
 
-#define UE_API CUSTOMIZABLEOBJECTEDITOR_API
-
 class UMaterialInterface;
 
 UENUM()
@@ -21,66 +21,52 @@ enum class ECustomizableObjectSelectionOverride : uint8
 };
 
 
-USTRUCT()
-struct FBoneToRemove
-{
-	GENERATED_USTRUCT_BODY()
-
-	UPROPERTY(EditAnywhere, Category = CustomizableObject)
-	bool bOnlyRemoveChildren = false;
-
-	UPROPERTY(EditAnywhere, Category = CustomizableObject)
-	FName BoneName;
-};
-
-
-USTRUCT()
-struct FLODReductionSettings
-{
-	GENERATED_USTRUCT_BODY()
-
-	/** Selects which bones will be removed from the final skeleton
-	* BoneName: Name of the bone that will be removed. Its children will be removed too.
-	* Remove Only Children: If true, only the children of the selected bone will be removed. The selected bone will remain.
-	*/
-	UPROPERTY(EditAnywhere, Category = CustomizableObject)
-	TArray<FBoneToRemove> BonesToRemove;
-};
-
-
-UCLASS(MinimalAPI)
+UCLASS()
 class UCustomizableObjectNodeComponentMesh : public UCustomizableObjectNodeComponent, public ICustomizableObjectNodeComponentMeshInterface
 {
 	GENERATED_BODY()
 
 public:
-	UE_API UCustomizableObjectNodeComponentMesh();
+	
+	UCustomizableObjectNodeComponentMesh();
 	
 	// UObject interface
-	UE_API virtual void PostEditChangeProperty(FPropertyChangedEvent& PropertyChangedEvent) override;
-	UE_API virtual void BackwardsCompatibleFixup(int32 CustomizableObjectCustomVersion) override;
+	virtual void BackwardsCompatibleFixup(int32 CustomizableObjectCustomVersion) override;
 
 	// UEdGraphNode interface
-	UE_API virtual void AllocateDefaultPins(UCustomizableObjectNodeRemapPins* RemapPins) override;
-	UE_API virtual FText GetNodeTitle(ENodeTitleType::Type TitleType) const override;
+	virtual void AllocateDefaultPins(UCustomizableObjectNodeRemapPins* RemapPins) override;
+	virtual FText GetNodeTitle(ENodeTitleType::Type TitleType) const override;
 
 	//UCustomizableObjectNode interface
-	UE_API bool IsSingleOutputNode() const;
-
+	virtual bool IsSingleOutputNode() const override;
+	virtual bool CanConnect(const UEdGraphPin* InOwnedInputPin, const UEdGraphPin* InOutputPin, bool& bOutIsOtherNodeBlocklisted,
+		bool& bOutArePinsCompatible) const override;
+	
 	// Own interface
-	UE_API FEdGraphPinReference GetOverlayMaterialAssetPin() const;
+	FEdGraphPinReference GetOverlayMaterialAssetPin() const;
 
 	// ComponentMesh Interface
-	UE_API virtual int32 GetNumLODs() override;
-	UE_API virtual const TArray<FEdGraphPinReference>& GetLODPins() const override;
-	UE_API virtual UEdGraphPin* GetOutputPin() const;
-	UE_API virtual void SetOutputPin(const UEdGraphPin* Pin);
-	UE_API virtual const UCustomizableObjectNode* GetOwningNode() const;
-
-public:
-
-	/*UPROPERTY(EditAnywhere, Category = ComponentMesh)
-	FName ComponentName = "Default name";*/
+	virtual int32 GetNumLODs() override;
+	virtual const TArray<FEdGraphPinReference>& GetLODPins() const override;
+	virtual UEdGraphPin* GetOutputPin() const override;
+	virtual void SetOutputPin(const UEdGraphPin* Pin) override;
+	virtual const UCustomizableObjectNode* GetOwningNode() const override;
+	
+	// Methods bellow are designed to provide the data this node before hosted but that now lives inside the node connected to the SKM pin.
+
+	/**
+	 * Get a pointer to the connected UCONodeSkeletalMeshMake node.
+	 * @param MacroContext The current macro stack
+	 * @return A valid pointer if the node is connected, a null pointer otherwise.
+	 */
+	const UCONodeSkeletalMeshMake* GetConnectedSkeletalMeshMakeNode(TArray<const UCustomizableObjectNodeMacroInstance*>& MacroContext) const;
+	
+	/**
+	 * Get a pointer to the connected UCONodeMutableSkeletalMeshMake node.
+	 * @param MacroContext The current macro stack
+	 * @return A valid pointer if the node is connected, a null pointer otherwise.
+	 */
+	const UCONodeMutableSkeletalMeshMake* GetConnectedMutableSkeletalMeshMakeNode(TArray<const UCustomizableObjectNodeMacroInstance*>& MacroContext) const;
 	
 	/** All the Skeletal Meshes generated for this component will use the Reference Skeletal Mesh properties
 	 * for everything Mutable doesn't create or modify. This includes data like LOD distances, Physics
@@ -92,26 +78,24 @@ public:
 	TObjectPtr<USkeletalMesh> ReferenceSkeletalMesh;
 
 	/** LOD reduction settings to apply on top of the LOD settings of the Reference Skeletal Mesh. */
-	UPROPERTY(EditAnywhere, Category = ComponentMesh)
-	TArray<FLODReductionSettings> LODReductionSettings;
+	UPROPERTY()
+	TArray<FLODReductionSettings> LODReductionSettings_DEPRECATED;
 
 	UPROPERTY()
 	TSoftObjectPtr<UMaterialInterface> OverlayMaterial_DEPRECATED;
 
-	UPROPERTY(EditAnywhere, Category = ComponentMesh)
-	FMutableLODSettings LODSettings;
-
-	/** Details selected LOD. */
-	int32 SelectedLOD = 0;
+	UPROPERTY()
+	FMutableLODSettings LODSettings_DEPRECATED;
 
-	UPROPERTY(EditAnywhere, Category = ComponentMesh)
-	int32 NumLODs = 1;
+	UPROPERTY()
+	int32 NumLODs_DEPRECATED = 1;
 
 	UPROPERTY()
-	TArray<FEdGraphPinReference> LODPins;
+	FEdGraphPinReference SkeletalMeshPin;
+	
+	UPROPERTY()
+	TArray<FEdGraphPinReference> LODPins_DEPRECATED;
 
 	UPROPERTY()
 	FEdGraphPinReference OverlayMaterialPin;
 };
-
-#undef UE_API
diff --git a/Engine/Plugins/Mutable/Source/CustomizableObjectEditor/Private/MuCOE/Nodes/CustomizableObjectNodeComponentMeshAddTo.h b/Engine/Plugins/Mutable/Source/CustomizableObjectEditor/Private/MuCOE/Nodes/CustomizableObjectNodeComponentMeshAddTo.h
--- a/Engine/Plugins/Mutable/Source/CustomizableObjectEditor/Private/MuCOE/Nodes/CustomizableObjectNodeComponentMeshAddTo.h
+++ b/Engine/Plugins/Mutable/Source/CustomizableObjectEditor/Private/MuCOE/Nodes/CustomizableObjectNodeComponentMeshAddTo.h
@@ -50,7 +50,7 @@ private:
 	FName ParentComponentName;
 
 public:
-	UPROPERTY(EditAnywhere, Category = ComponentMesh)
+	UPROPERTY(EditAnywhere, Category = ComponentMesh, meta=(CLampMin=1))
 	int32 NumLODs = 1;
 
 	UPROPERTY()
diff --git a/Engine/Plugins/Mutable/Source/CustomizableObjectEditor/Private/MuCOE/Nodes/CustomizableObjectNodeComponentMeshDetails.h b/Engine/Plugins/Mutable/Source/CustomizableObjectEditor/Private/MuCOE/Nodes/CustomizableObjectNodeComponentMeshDetails.h
deleted file mode 100644
diff --git a/Engine/Plugins/Mutable/Source/CustomizableObjectEditor/Private/MuCOE/Nodes/CustomizableObjectNodeObject.cpp b/Engine/Plugins/Mutable/Source/CustomizableObjectEditor/Private/MuCOE/Nodes/CustomizableObjectNodeObject.cpp
--- a/Engine/Plugins/Mutable/Source/CustomizableObjectEditor/Private/MuCOE/Nodes/CustomizableObjectNodeObject.cpp
+++ b/Engine/Plugins/Mutable/Source/CustomizableObjectEditor/Private/MuCOE/Nodes/CustomizableObjectNodeObject.cpp
@@ -199,7 +199,7 @@ void UCustomizableObjectNodeObject::BackwardsCompatibleFixup(int32 CustomizableO
 		{
 			// UE Code from FSchemaAction_NewNode::CreateNode(...). Overlap calculations performed before AutowireNewNode(...).
 
-			const int32 NodeDistance = 60;
+			constexpr int32 NodeDistance = 60;
 			
 			// Duplicate template node to create new node
 			UEdGraphNode* ResultNode = DuplicateObject<UEdGraphNode>(InNodeTemplate, ParentGraph);
@@ -272,13 +272,13 @@ void UCustomizableObjectNodeObject::BackwardsCompatibleFixup(int32 CustomizableO
 				++NodesCreated;
 				
 				NodeComponentMesh->SetComponentName(MeshComponent.Name);
-				NodeComponentMesh->NumLODs = NumLODs_DEPRECATED;
+				NodeComponentMesh->NumLODs_DEPRECATED = NumLODs_DEPRECATED;
 				if (FComponentSettings* Result = ComponentSettings_DEPRECATED.FindByPredicate([&](const FComponentSettings& Settings)
 				{
 					return Settings.ComponentName == MeshComponent.Name;
 				}))
 				{
-					NodeComponentMesh->LODReductionSettings = Result->LODReductionSettings;
+					NodeComponentMesh->LODReductionSettings_DEPRECATED = Result->LODReductionSettings;
 				}
 
 				NodeComponentMesh->FixupReconstructPins(CreateRemapPinsByName(), NodeComponentMeshBaseAllocateDefaultPins);
@@ -501,16 +501,16 @@ void UCustomizableObjectNodeObject::BackwardsCompatibleFixup(int32 CustomizableO
 					}
 
 					UCustomizableObjectNodeComponentMesh* NodeComponentMesh = *Result;
-					if (NodeComponentMesh->NumLODs != NodeComponentMeshAddTo->NumLODs)
+					if (NodeComponentMesh->NumLODs_DEPRECATED != NodeComponentMeshAddTo->NumLODs)
 					{
 						continue;
 					}
 				
-					for (int32 LODIndex = 0; LODIndex < NodeComponentMesh->NumLODs; ++LODIndex)
+					for (int32 LODIndex = 0; LODIndex < NodeComponentMesh->NumLODs_DEPRECATED; ++LODIndex)
 					{
 						for (UEdGraphPin* LinkedPin : NodeComponentMeshAddTo->LODPins[LODIndex].Get()->LinkedTo)
 						{
-							LinkedPin->MakeLinkTo(NodeComponentMesh->LODPins[LODIndex].Get());
+							LinkedPin->MakeLinkTo(NodeComponentMesh->LODPins_DEPRECATED[LODIndex].Get());
 						}
 					}
 
diff --git a/Engine/Plugins/Mutable/Source/CustomizableObjectEditor/Private/MuCOE/Nodes/CustomizableObjectNodeTextureLayer.cpp b/Engine/Plugins/Mutable/Source/CustomizableObjectEditor/Private/MuCOE/Nodes/CustomizableObjectNodeTextureLayer.cpp
--- a/Engine/Plugins/Mutable/Source/CustomizableObjectEditor/Private/MuCOE/Nodes/CustomizableObjectNodeTextureLayer.cpp
+++ b/Engine/Plugins/Mutable/Source/CustomizableObjectEditor/Private/MuCOE/Nodes/CustomizableObjectNodeTextureLayer.cpp
@@ -127,7 +127,7 @@ bool UCustomizableObjectNodeTextureLayer::CanConnect(const UEdGraphPin* InOwnedI
 	const UEdGraphSchema_CustomizableObject* Schema = GetDefault<UEdGraphSchema_CustomizableObject>();
 
 	// In layer pins we accept both colors and images
-	for ( int LayerIndex = 0; LayerIndex < Layers.Num(); ++LayerIndex )
+	for ( int32 LayerIndex = 0; LayerIndex < Layers.Num(); ++LayerIndex )
 	{
 		if (LayerPin(LayerIndex) == InOwnedInputPin )
 		{
diff --git a/Engine/Plugins/Mutable/Source/CustomizableObjectEditor/Private/MuCOE/Nodes/FCustomizableObjectNodeComponentMeshDetails.cpp b/Engine/Plugins/Mutable/Source/CustomizableObjectEditor/Private/MuCOE/Nodes/FCustomizableObjectNodeComponentMeshDetails.cpp
deleted file mode 100644
