[mutable] Renamed the string input pin from "String " to "Name"

#Synced-CL 51817239

diff --git a/Engine/Plugins/Mutable/Source/CustomizableObjectEditor/Private/MuCOE/Nodes/CONodeSkeletalMeshMake.cpp b/Engine/Plugins/Mutable/Source/CustomizableObjectEditor/Private/MuCOE/Nodes/CONodeSkeletalMeshMake.cpp
--- a/Engine/Plugins/Mutable/Source/CustomizableObjectEditor/Private/MuCOE/Nodes/CONodeSkeletalMeshMake.cpp
+++ b/Engine/Plugins/Mutable/Source/CustomizableObjectEditor/Private/MuCOE/Nodes/CONodeSkeletalMeshMake.cpp
@@ -35,7 +35,7 @@ FLinearColor UCONodeSkeletalMeshMake::GetNodeTitleColor() const
 
 void UCONodeSkeletalMeshMake::AllocateDefaultPins(UCustomizableObjectNodeRemapPins* RemapPins)
 {
-	MeshNamePin = CustomCreatePin(EEdGraphPinDirection::EGPD_Input, UEdGraphSchema_CustomizableObject::PC_String);
+	MeshNamePin = CustomCreatePin(EEdGraphPinDirection::EGPD_Input, UEdGraphSchema_CustomizableObject::PC_String, TEXT("Name"), LOCTEXT("Name", "Name"));
 	SkeletalMeshPin = CustomCreatePin(EEdGraphPinDirection::EGPD_Input, UEdGraphSchema_CustomizableObject::PC_SkeletalMesh);
 	
 	LODReductionSettings.SetNum(NumLODs);
