[Mutable] Allow killing UpdateImage tasks.

#rnx
#Synced-CL 51880525

diff --git a/Engine/Plugins/Mutable/Source/CustomizableObject/Private/MuCO/CustomizableObjectMipDataProvider.cpp b/Engine/Plugins/Mutable/Source/CustomizableObject/Private/MuCO/CustomizableObjectMipDataProvider.cpp
--- a/Engine/Plugins/Mutable/Source/CustomizableObject/Private/MuCO/CustomizableObjectMipDataProvider.cpp
+++ b/Engine/Plugins/Mutable/Source/CustomizableObject/Private/MuCO/CustomizableObjectMipDataProvider.cpp
@@ -77,6 +77,11 @@ namespace Impl
 	void Task_Mutable_UpdateImage(TSharedPtr<FMutableImageOperationData> OperationData)
 	{
 		MUTABLE_CPUPROFILER_SCOPE(Task_Mutable_UpdateImage)
+	
+		if (OperationData->bIsCancelled)
+		{
+			return;
+		}
 		
 		const double StartTime = FPlatformTime::Seconds();
 		
