[Mutable] Avoid exploring unecessary branches on BeginUpdate.

#Synced-CL 51880525

diff --git a/Engine/Plugins/Mutable/Source/MutableRuntime/Private/MuR/CodeRunner.cpp b/Engine/Plugins/Mutable/Source/MutableRuntime/Private/MuR/CodeRunner.cpp
--- a/Engine/Plugins/Mutable/Source/MutableRuntime/Private/MuR/CodeRunner.cpp
+++ b/Engine/Plugins/Mutable/Source/MutableRuntime/Private/MuR/CodeRunner.cpp
@@ -70,42 +70,42 @@ namespace
 {
 
 int32 ForcedProjectionMode = -1;
-static FAutoConsoleVariableRef CVarForceProjectionSamplingMode (
+FAutoConsoleVariableRef CVarForceProjectionSamplingMode (
 	TEXT("mutable.ForceProjectionMode"),
 	ForcedProjectionMode,
 	TEXT("force mutable to use an specific projection mode, 0 = Point + None, 1 = Bilinear + TotalAreaHeuristic, -1 uses the values provided by the projector."),
 	ECVF_Default);
 
 float GlobalProjectionLodBias = 0.0f;
-static FAutoConsoleVariableRef CVarGlobalProjectionLodBias (
+FAutoConsoleVariableRef CVarGlobalProjectionLodBias (
 	TEXT("mutable.GlobalProjectionLodBias"),
 	GlobalProjectionLodBias,
 	TEXT("Lod bias applied to the lod resulting form the best mip computation for ImageProject operations, only used if a min filter method different than None is used."),
 	ECVF_Default);
 
 bool bUseProjectionVectorImpl = true;
-static FAutoConsoleVariableRef CVarUseProjectionVectorImpl (
+FAutoConsoleVariableRef CVarUseProjectionVectorImpl (
 	TEXT("mutable.UseProjectionVectorImpl"),
 	bUseProjectionVectorImpl,
 	TEXT("If set to true, enables the vectorized implementation of the projection pixel processing."),
 	ECVF_Default);
 
 float GlobalImageTransformLodBias = 0.0f;
-static FAutoConsoleVariableRef CVarGlobalImageTransformLodBias (
+FAutoConsoleVariableRef CVarGlobalImageTransformLodBias (
 	TEXT("mutable.GlobalImageTransformLodBias"),
 	GlobalImageTransformLodBias,
 		TEXT("Lod bias applied to the lod resulting form the best mip computation for ImageTransform operations"),
 	ECVF_Default);
 
 bool bUseImageTransformVectorImpl = true;
-static FAutoConsoleVariableRef CVarUseImageTransformVectorImpl (
+FAutoConsoleVariableRef CVarUseImageTransformVectorImpl (
 	TEXT("mutable.UseImageTransformVectorImpl"),
 	bUseImageTransformVectorImpl,
 	TEXT("If set to true, enables the vectorized implementation of the image transform pixel processing."),
 	ECVF_Default);
 
 int32 MaxWorkerThreadsForIterativeOpExec = 4;
-static FAutoConsoleVariableRef CVarMaxWorkerThreadsForIterativeOpExec(
+FAutoConsoleVariableRef CVarMaxWorkerThreadsForIterativeOpExec(
 	TEXT("mutable.MaxWorkerThreadsForIterativeOpExec"),
 	MaxWorkerThreadsForIterativeOpExec,
 	TEXT("Sets the maximum number of worker threads to consider executing the iterative variant of an operation.")
@@ -871,10 +871,8 @@ namespace UE::Mutable::Private
         		{
         		case 0:
         			{
-        				AddOp(FScheduledOp(Item.At, Item, 1),
-							FScheduledOp(Args.Instance, Item),
-							FScheduledOp(Args.Material, Item));
-
+        				AddOp(FScheduledOp(Item.At, Item, 1), 
+        					FScheduledOp(Args.Instance, Item));
         				break;
         			}
 
@@ -927,9 +925,7 @@ namespace UE::Mutable::Private
         		case 0:
         			{
         				AddOp(FScheduledOp(Item.At, Item, 1),
-							FScheduledOp(Args.Instance, Item),
-							FScheduledOp(Args.Material, Item));
-
+        					FScheduledOp(Args.Instance, Item));
         				break;
         			}
 
@@ -973,8 +969,7 @@ namespace UE::Mutable::Private
         		{
         		case 0:
         			AddOp(FScheduledOp(Item.At, Item, 1),
-						FScheduledOp(Args.Instance, Item),
-						FScheduledOp(Args.Material, Item));
+						FScheduledOp(Args.Instance, Item));
         			break;
 
         		case 1:
@@ -1018,10 +1013,8 @@ namespace UE::Mutable::Private
         		{
         		case 0:
         			{
-        				AddOp(FScheduledOp(Item.At, Item, 1),
-							FScheduledOp(Args.instance, Item),
-							FScheduledOp(Args.value, Item));
-						
+        				AddOp(FScheduledOp(Item.At, Item, 1), 
+        					FScheduledOp(Args.instance, Item));
         				break;
         			}
 
