[Mutable] Fix Code Viewer stack overflow crash.

#Synced-CL 51880525

diff --git a/Engine/Plugins/Mutable/Source/CustomizableObjectEditor/Private/MuCOE/SMutableCodeViewer.cpp b/Engine/Plugins/Mutable/Source/CustomizableObjectEditor/Private/MuCOE/SMutableCodeViewer.cpp
--- a/Engine/Plugins/Mutable/Source/CustomizableObjectEditor/Private/MuCOE/SMutableCodeViewer.cpp
+++ b/Engine/Plugins/Mutable/Source/CustomizableObjectEditor/Private/MuCOE/SMutableCodeViewer.cpp
@@ -16,7 +16,6 @@
 #include "MuCOE/SMutableIntViewer.h"
 #include "MuCOE/SMutableLayoutViewer.h"
 #include "MuCOE/SMutableMeshViewer.h"
-#include "MuCOE/SMutableInstanceViewer.h"
 #include "MuCOE/SMutableProjectorViewer.h"
 #include "MuCOE/SMutableScalarViewer.h"
 #include "MuCOE/SMutableSkeletonViewer.h"
@@ -1349,530 +1348,606 @@ void SMutableCodeViewer::GenerateAllTreeElements()
 			++ModelOperationTypes[(int32)RootType].Value;
 
 			// Iterate over each root node and generate all the elements in a human-readable pattern (Z Pattern)
-			GenerateElementRecursive(StateIndex, State.Root, Program);
+			GenerateElement(StateIndex, State.Root, Program);
 		}
 	}
 	
 }
 
-void SMutableCodeViewer::GenerateElementRecursive(const int32& InStateIndex, UE::Mutable::Private::FOperation::ADDRESS InParentAddress,  const UE::Mutable::Private::FProgram& InProgram)
+void SMutableCodeViewer::GenerateElement(const int32& InStateIndex, UE::Mutable::Private::FOperation::ADDRESS InParentAddress,  const UE::Mutable::Private::FProgram& InProgram)
 {
+	TQueue<UE::Mutable::Private::FOperation::ADDRESS> Addresses;
+	Addresses.Enqueue(InParentAddress);
+	
 	// This will be used to add operations
-	uint32 ChildIndex = 0;
-	auto AddOpFunc = [this, InParentAddress, &InProgram, &ChildIndex, &InStateIndex](UE::Mutable::Private::FOperation::ADDRESS ChildAddress, const FString& Caption)
+	auto AddOpFunc = [this, &InProgram, &InStateIndex, &Addresses](UE::Mutable::Private::FOperation::ADDRESS ParentAddress, UE::Mutable::Private::FOperation::ADDRESS ChildAddress, const FString& Caption, uint32& ChildIndex)
+	{
+		const FItemCacheKey Key = { ParentAddress, ChildAddress, ChildIndex };
+		const TSharedPtr<FMutableCodeTreeElement>* CachedItem = ItemCache.Find(Key);
+			
+		ChildIndex++;
+
+		// If not already cached then process it
+		if (ensure(!CachedItem))
 		{
+			// Locate the "original" tree element
+			const TSharedPtr<FMutableCodeTreeElement>* MainItemPtr = MainItemPerOp.Find(ChildAddress);
+
+			// Provide the color this element should be using for the displayed text
+			const FSlateColor LabelColor = ColorPerComputationalCost[StaticCast<uint8>(GetOperationTypeComputationalCost(InProgram.GetOpType(ChildAddress)))];
+
+			const TSharedPtr<FMutableCodeTreeElement> Item = MakeShareable(new FMutableCodeTreeElement(ItemCache.Num(),InStateIndex, MutableModel, ChildAddress, Caption, LabelColor, MainItemPtr));
+
+			// Cache this element for later access
+			ItemCache.Add(Key, Item);
+
+			// It is not a duplicated of another one, then we can continue searching
+			if (!MainItemPtr)
 			{
-				const FItemCacheKey Key = { InParentAddress, ChildAddress, ChildIndex };
-				const TSharedPtr<FMutableCodeTreeElement>* CachedItem = ItemCache.Find(Key);
+				MainItemPerOp.Add(ChildAddress, Item);
 
-				// If not already cached then process it
-				if (ensure(!CachedItem))
-				{
-					// Locate the "original" tree element
-					const TSharedPtr<FMutableCodeTreeElement>* MainItemPtr = MainItemPerOp.Find(ChildAddress);
+				// Populate OperatationTypes with the root.
+				UE::Mutable::Private::EOpType ChildType = InProgram.GetOpType(ChildAddress);
+				ModelOperationTypes[(int32)ChildType].Key = ChildType;
+				++ModelOperationTypes[(int32)ChildType].Value;
+				
+				Addresses.Enqueue( ChildAddress );
+			}
+		}
+		else
+		{
+			UE_LOG(LogMutable, Error, TEXT("An already processed operation is being re-processed in order to generate a tree row. "
+							"ParentAddress : %u , ChildAddress : %u , ChildIndex : %u "), ParentAddress, ChildAddress, ChildIndex);
+		}
+	};
+	
+	while (!Addresses.IsEmpty())
+	{
+		UE::Mutable::Private::FOperation::ADDRESS Address;
+		Addresses.Dequeue(Address);
+		
+		// For some specific parent operation types we create more detailed subtrees.
+		bool bUseGeneric = false;
+		const UE::Mutable::Private::EOpType ParentOperationType = InProgram.GetOpType(Address);
+		switch (ParentOperationType)
+		{
+		case UE::Mutable::Private::EOpType::IM_CONDITIONAL:
+		case UE::Mutable::Private::EOpType::LA_CONDITIONAL:
+		case UE::Mutable::Private::EOpType::ME_CONDITIONAL:
+		case UE::Mutable::Private::EOpType::CO_CONDITIONAL:
+		case UE::Mutable::Private::EOpType::SC_CONDITIONAL:
+		case UE::Mutable::Private::EOpType::NU_CONDITIONAL:
+		case UE::Mutable::Private::EOpType::IN_CONDITIONAL:
+		case UE::Mutable::Private::EOpType::ED_CONDITIONAL:
+		case UE::Mutable::Private::EOpType::MI_CONDITIONAL:
+			{
+				uint32 ChildIndex = 0;
+			
+				UE::Mutable::Private::FOperation::ConditionalArgs Args = InProgram.GetOpArgs<UE::Mutable::Private::FOperation::ConditionalArgs>(Address);
+				AddOpFunc(Address, Args.condition, TEXT("cond "), ChildIndex);
+				AddOpFunc(Address, Args.yes, TEXT("true "), ChildIndex);
+				AddOpFunc(Address, Args.no, TEXT("false "), ChildIndex);
+				break;
+			}
+
+		case UE::Mutable::Private::EOpType::IM_SWITCH:
+		case UE::Mutable::Private::EOpType::LA_SWITCH:
+		case UE::Mutable::Private::EOpType::ME_SWITCH:
+		case UE::Mutable::Private::EOpType::CO_SWITCH:
+		case UE::Mutable::Private::EOpType::SC_SWITCH:
+		case UE::Mutable::Private::EOpType::NU_SWITCH:
+		case UE::Mutable::Private::EOpType::IN_SWITCH:
+		case UE::Mutable::Private::EOpType::ED_SWITCH:
+			{
+				uint32 ChildIndex = 0;
 
-					// Provide the color this element should be using for the displayed text
-					const FSlateColor LabelColor = ColorPerComputationalCost[StaticCast<uint8>(GetOperationTypeComputationalCost(InProgram.GetOpType(ChildAddress)))];
+				const uint8* OpData = InProgram.GetOpArgsPointer(Address);
 
-					const TSharedPtr<FMutableCodeTreeElement> Item = MakeShareable(new FMutableCodeTreeElement(ItemCache.Num(),InStateIndex, MutableModel, ChildAddress, Caption, LabelColor, MainItemPtr));
+				UE::Mutable::Private::FOperation::ADDRESS VarAddress;
+				FMemory::Memcpy(&VarAddress, OpData, sizeof(UE::Mutable::Private::FOperation::ADDRESS));
+				OpData += sizeof(UE::Mutable::Private::FOperation::ADDRESS);
+				AddOpFunc(Address, VarAddress, TEXT("var "), ChildIndex);
 
-					// Cache this element for later access
-					ItemCache.Add(Key, Item);
+				UE::Mutable::Private::FOperation::ADDRESS DefAddress;
+				FMemory::Memcpy(&DefAddress, OpData, sizeof(UE::Mutable::Private::FOperation::ADDRESS));
+				OpData += sizeof(UE::Mutable::Private::FOperation::ADDRESS);
+				AddOpFunc(Address, DefAddress, TEXT("def "), ChildIndex);
 
-					// It is not a duplicated of another one, then we can continue searching
-					if (!MainItemPtr)
+				UE::Mutable::Private::FOperation::FSwitchCaseDescriptor CaseDesc;
+				FMemory::Memcpy(&CaseDesc, OpData, sizeof(UE::Mutable::Private::FOperation::FSwitchCaseDescriptor));
+				OpData += sizeof(UE::Mutable::Private::FOperation::FSwitchCaseDescriptor);
+
+				if (!CaseDesc.bUseRanges)
+				{
+					for (uint32 C = 0; C < CaseDesc.Count; ++C)
 					{
-						MainItemPerOp.Add(ChildAddress, Item);
+						int32 Condition;
+						FMemory::Memcpy(&Condition, OpData, sizeof(int32));
+						OpData += sizeof(int32);
 
-						// Populate OperatationTypes with the root.
-						UE::Mutable::Private::EOpType ChildType = InProgram.GetOpType(ChildAddress);
-						ModelOperationTypes[(int32)ChildType].Key = ChildType;
-						++ModelOperationTypes[(int32)ChildType].Value;
+						UE::Mutable::Private::FOperation::ADDRESS At;
+						FMemory::Memcpy(&At, OpData, sizeof(UE::Mutable::Private::FOperation::ADDRESS));
+						OpData += sizeof(UE::Mutable::Private::FOperation::ADDRESS);
 
-						GenerateElementRecursive(InStateIndex,ChildAddress, InProgram);
+						FString Caption = FString::Printf(TEXT("case %d "), Condition);
+						AddOpFunc(Address, At, Caption, ChildIndex);
 					}
 				}
 				else
 				{
-					UE_LOG(LogMutable, Error, TEXT("An already processed operation is being re-processed in order to generate a tree row. "
-									"ParentAddress : %u , ChildAddress : %u , ChildIndex : %u "), InParentAddress, ChildAddress, ChildIndex);
-				}
-				
-			}
-			++ChildIndex;
-		};
-
-	
-	// For some specific parent operation types we create more detailed subtrees.
-	bool bUseGeneric = false;
-	const UE::Mutable::Private::EOpType ParentOperationType = InProgram.GetOpType(InParentAddress);
-	switch (ParentOperationType)
-	{
-	case UE::Mutable::Private::EOpType::IM_CONDITIONAL:
-	case UE::Mutable::Private::EOpType::LA_CONDITIONAL:
-	case UE::Mutable::Private::EOpType::ME_CONDITIONAL:
-	case UE::Mutable::Private::EOpType::CO_CONDITIONAL:
-	case UE::Mutable::Private::EOpType::SC_CONDITIONAL:
-	case UE::Mutable::Private::EOpType::NU_CONDITIONAL:
-	case UE::Mutable::Private::EOpType::IN_CONDITIONAL:
-	case UE::Mutable::Private::EOpType::ED_CONDITIONAL:
-	{
-		UE::Mutable::Private::FOperation::ConditionalArgs Args = InProgram.GetOpArgs<UE::Mutable::Private::FOperation::ConditionalArgs>(InParentAddress);
-		AddOpFunc(Args.condition, TEXT("cond "));
-		AddOpFunc(Args.yes, TEXT("true "));
-		AddOpFunc(Args.no, TEXT("false "));
-		break;
-	}
+					for (uint32 C = 0; C < CaseDesc.Count; ++C)
+					{
+						int32 ConditionStart;
+						FMemory::Memcpy(&ConditionStart, OpData, sizeof(int32));
+						OpData += sizeof(int32);
 
-	case UE::Mutable::Private::EOpType::IM_SWITCH:
-	case UE::Mutable::Private::EOpType::LA_SWITCH:
-	case UE::Mutable::Private::EOpType::ME_SWITCH:
-	case UE::Mutable::Private::EOpType::CO_SWITCH:
-	case UE::Mutable::Private::EOpType::SC_SWITCH:
-	case UE::Mutable::Private::EOpType::NU_SWITCH:
-	case UE::Mutable::Private::EOpType::IN_SWITCH:
-	case UE::Mutable::Private::EOpType::ED_SWITCH:
-	{
-		const uint8* OpData = InProgram.GetOpArgsPointer(InParentAddress);
+						uint32 RangeSize;
+						FMemory::Memcpy(&RangeSize, OpData, sizeof(uint32));
+						OpData += sizeof(uint32);
 
-		UE::Mutable::Private::FOperation::ADDRESS VarAddress;
-		FMemory::Memcpy(&VarAddress, OpData, sizeof(UE::Mutable::Private::FOperation::ADDRESS));
-		OpData += sizeof(UE::Mutable::Private::FOperation::ADDRESS);
-		AddOpFunc(VarAddress, TEXT("var "));
+						UE::Mutable::Private::FOperation::ADDRESS At;
+						FMemory::Memcpy(&At, OpData, sizeof(UE::Mutable::Private::FOperation::ADDRESS));
+						OpData += sizeof(UE::Mutable::Private::FOperation::ADDRESS);
 
-		UE::Mutable::Private::FOperation::ADDRESS DefAddress;
-		FMemory::Memcpy(&DefAddress, OpData, sizeof(UE::Mutable::Private::FOperation::ADDRESS));
-		OpData += sizeof(UE::Mutable::Private::FOperation::ADDRESS);
-		AddOpFunc(DefAddress, TEXT("def "));
+						FString Caption = FString::Printf(TEXT("case %d to %d "), ConditionStart, ConditionStart+RangeSize-1);
+						AddOpFunc(Address, At, Caption, ChildIndex);
+					}
+				}
 
-		UE::Mutable::Private::FOperation::FSwitchCaseDescriptor CaseDesc;
-		FMemory::Memcpy(&CaseDesc, OpData, sizeof(UE::Mutable::Private::FOperation::FSwitchCaseDescriptor));
-		OpData += sizeof(UE::Mutable::Private::FOperation::FSwitchCaseDescriptor);
+				break;
+			}
 
-		if (!CaseDesc.bUseRanges)
-		{
-			for (uint32 C = 0; C < CaseDesc.Count; ++C)
+		case UE::Mutable::Private::EOpType::IM_SWIZZLE:
 			{
-				int32 Condition;
-				FMemory::Memcpy(&Condition, OpData, sizeof(int32));
-				OpData += sizeof(int32);
+				uint32 ChildIndex = 0;
 
-				UE::Mutable::Private::FOperation::ADDRESS At;
-				FMemory::Memcpy(&At, OpData, sizeof(UE::Mutable::Private::FOperation::ADDRESS));
-				OpData += sizeof(UE::Mutable::Private::FOperation::ADDRESS);
-
-				FString Caption = FString::Printf(TEXT("case %d "), Condition);
-				AddOpFunc(At, Caption);
+				UE::Mutable::Private::FOperation::ImageSwizzleArgs Args = InProgram.GetOpArgs<UE::Mutable::Private::FOperation::ImageSwizzleArgs>(Address);
+				for (int32 Channel = 0; Channel < 4; ++Channel)
+				{
+					FString Caption = FString::Printf(TEXT("%d is %d from "), Channel, Args.sourceChannels[Channel]);
+					AddOpFunc(Address, Args.sources[Channel], Caption, ChildIndex);
+				}
+				break;
 			}
-		}
-		else
-		{
-			for (uint32 C = 0; C < CaseDesc.Count; ++C)
+
+		case UE::Mutable::Private::EOpType::CO_SWIZZLE:
 			{
-				int32 ConditionStart;
-				FMemory::Memcpy(&ConditionStart, OpData, sizeof(int32));
-				OpData += sizeof(int32);
+				uint32 ChildIndex = 0;
+			
+				UE::Mutable::Private::FOperation::ColourSwizzleArgs Args = InProgram.GetOpArgs<UE::Mutable::Private::FOperation::ColourSwizzleArgs>(Address);
+				for (int32 Channel = 0; Channel < 4; ++Channel)
+				{
+					FString Caption = FString::Printf(TEXT("%d is %d from "), Channel, Args.sourceChannels[Channel]);
+					AddOpFunc(Address, Args.sources[Channel], Caption, ChildIndex);
+				}
+				break;
+			}
 
-				uint32 RangeSize;
-				FMemory::Memcpy(&RangeSize, OpData, sizeof(uint32));
-				OpData += sizeof(uint32);
+		case UE::Mutable::Private::EOpType::IM_LAYER:
+			{
+				uint32 ChildIndex = 0;
+			
+				UE::Mutable::Private::FOperation::ImageLayerArgs Args = InProgram.GetOpArgs<UE::Mutable::Private::FOperation::ImageLayerArgs>(Address);
+				AddOpFunc(Address, Args.base, TEXT("base "), ChildIndex);
+				AddOpFunc(Address, Args.mask, TEXT("mask "), ChildIndex);
+				AddOpFunc(Address, Args.blended, TEXT("blended "), ChildIndex);
+				break;
+			}
 
-				UE::Mutable::Private::FOperation::ADDRESS At;
-				FMemory::Memcpy(&At, OpData, sizeof(UE::Mutable::Private::FOperation::ADDRESS));
-				OpData += sizeof(UE::Mutable::Private::FOperation::ADDRESS);
+		case UE::Mutable::Private::EOpType::IM_LAYERCOLOUR:
+			{
+				uint32 ChildIndex = 0;
+			
+				UE::Mutable::Private::FOperation::ImageLayerColourArgs Args = InProgram.GetOpArgs<UE::Mutable::Private::FOperation::ImageLayerColourArgs>(Address);
+				AddOpFunc(Address, Args.base, TEXT("base "), ChildIndex);
+				AddOpFunc(Address, Args.mask, TEXT("mask "), ChildIndex);
+				AddOpFunc(Address, Args.colour, TEXT("colour "), ChildIndex);
+				break;
+			}
 
-				FString Caption = FString::Printf(TEXT("case %d to %d "), ConditionStart, ConditionStart+RangeSize-1);
-				AddOpFunc(At, Caption);
+		case UE::Mutable::Private::EOpType::IM_MULTILAYER:
+			{
+				uint32 ChildIndex = 0;
+			
+				UE::Mutable::Private::FOperation::ImageMultiLayerArgs Args = InProgram.GetOpArgs<UE::Mutable::Private::FOperation::ImageMultiLayerArgs>(Address);
+				AddOpFunc(Address, Args.rangeSize, TEXT("range "), ChildIndex);
+				AddOpFunc(Address, Args.base, TEXT("base "), ChildIndex);
+				AddOpFunc(Address, Args.mask, TEXT("mask "), ChildIndex);
+				AddOpFunc(Address, Args.blended, TEXT("blended "), ChildIndex);
+				break;
 			}
-		}
 
-		break;
-	}
+		case UE::Mutable::Private::EOpType::ME_ADDMETADATA:
+			{
+				uint32 ChildIndex = 0;
 
-	case UE::Mutable::Private::EOpType::IM_SWIZZLE:
-	{
-		UE::Mutable::Private::FOperation::ImageSwizzleArgs Args = InProgram.GetOpArgs<UE::Mutable::Private::FOperation::ImageSwizzleArgs>(InParentAddress);
-		for (int32 Channel = 0; Channel < 4; ++Channel)
-		{
-			FString Caption = FString::Printf(TEXT("%d is %d from "), Channel, Args.sourceChannels[Channel]);
-			AddOpFunc(Args.sources[Channel], Caption);
-		}
-		break;
-	}
+				UE::Mutable::Private::FOperation::MeshAddMetadataArgs Args = InProgram.GetOpArgs<UE::Mutable::Private::FOperation::MeshAddMetadataArgs>(Address);
 
-	case UE::Mutable::Private::EOpType::CO_SWIZZLE:
-	{
-		UE::Mutable::Private::FOperation::ColourSwizzleArgs Args = InProgram.GetOpArgs<UE::Mutable::Private::FOperation::ColourSwizzleArgs>(InParentAddress);
-		for (int32 Channel = 0; Channel < 4; ++Channel)
-		{
-			FString Caption = FString::Printf(TEXT("%d is %d from "), Channel, Args.sourceChannels[Channel]);
-			AddOpFunc(Args.sources[Channel], Caption);
-		}
-		break;
-	}
+				using OpEnumFlags = UE::Mutable::Private::FOperation::MeshAddMetadataArgs::EnumFlags;
 
-	case UE::Mutable::Private::EOpType::IM_LAYER:
-	{
-		UE::Mutable::Private::FOperation::ImageLayerArgs Args = InProgram.GetOpArgs<UE::Mutable::Private::FOperation::ImageLayerArgs>(InParentAddress);
-		AddOpFunc(Args.base, TEXT("base "));
-		AddOpFunc(Args.mask, TEXT("mask "));
-		AddOpFunc(Args.blended, TEXT("blended "));
-		break;
-	}
-
-	case UE::Mutable::Private::EOpType::IM_LAYERCOLOUR:
-	{
-		UE::Mutable::Private::FOperation::ImageLayerColourArgs Args = InProgram.GetOpArgs<UE::Mutable::Private::FOperation::ImageLayerColourArgs>(InParentAddress);
-		AddOpFunc(Args.base, TEXT("base "));
-		AddOpFunc(Args.mask, TEXT("mask "));
-		AddOpFunc(Args.colour, TEXT("colour "));
-		break;
-	}
-
-	case UE::Mutable::Private::EOpType::IM_MULTILAYER:
-	{
-		UE::Mutable::Private::FOperation::ImageMultiLayerArgs Args = InProgram.GetOpArgs<UE::Mutable::Private::FOperation::ImageMultiLayerArgs>(InParentAddress);
-		AddOpFunc(Args.rangeSize, TEXT("range "));
-		AddOpFunc(Args.base, TEXT("base "));
-		AddOpFunc(Args.mask, TEXT("mask "));
-		AddOpFunc(Args.blended, TEXT("blended "));
-		break;
-	}
+				int32 TagCount = EnumHasAnyFlags(Args.Flags, OpEnumFlags::HasGameplayTags) ? 1 : 0;
+				if (EnumHasAnyFlags(Args.Flags, OpEnumFlags::IsGameplayTagList) &&
+					ensure(InProgram.ConstantUInt32Lists.IsValidIndex(Args.GameplayTags.ListIndex)))
+				{
+					TagCount = InProgram.ConstantUInt32Lists.IsValidIndex(Args.GameplayTags.ListIndex) 
+							? InProgram.ConstantUInt32Lists[Args.GameplayTags.ListIndex].Num() 
+							: 0;
+				}
 
-	case UE::Mutable::Private::EOpType::ME_ADDMETADATA:
-	{
-		UE::Mutable::Private::FOperation::MeshAddMetadataArgs Args = InProgram.GetOpArgs<UE::Mutable::Private::FOperation::MeshAddMetadataArgs>(InParentAddress);
+				int32 ResourceCount = EnumHasAnyFlags(Args.Flags, OpEnumFlags::HasAnimationSlots) ? 1 : 0;
+				if (EnumHasAnyFlags(Args.Flags, OpEnumFlags::IsAnimationSlotList) &&
+					ensure(InProgram.ConstantUInt32Lists.IsValidIndex(Args.AnimSlotNames.ListIndex)))
+				{
+					ResourceCount = InProgram.ConstantUInt32Lists.IsValidIndex(Args.AnimSlotNames.ListIndex) 
+							? InProgram.ConstantUInt32Lists[Args.AnimSlotNames.ListIndex].Num() 
+							: 0;
+				}
 
-		using OpEnumFlags = UE::Mutable::Private::FOperation::MeshAddMetadataArgs::EnumFlags;
+				int32 RealTimeMorphNamesCount = EnumHasAnyFlags(Args.Flags, OpEnumFlags::HasRealTimeMorphNames) ? 1 : 0;
+				if (EnumHasAnyFlags(Args.Flags, OpEnumFlags::IsRealTimeMorphNamesList) &&
+					ensure(InProgram.ConstantUInt32Lists.IsValidIndex(Args.RealTimeMorphNames.ListAddress)))
+				{
+					RealTimeMorphNamesCount = InProgram.ConstantUInt32Lists.IsValidIndex(Args.RealTimeMorphNames.ListAddress) 
+							? InProgram.ConstantUInt32Lists[Args.RealTimeMorphNames.ListAddress].Num() 
+							: 0;
+				}
 
-		int32 TagCount = EnumHasAnyFlags(Args.Flags, OpEnumFlags::HasGameplayTags) ? 1 : 0;
-		if (EnumHasAnyFlags(Args.Flags, OpEnumFlags::IsGameplayTagList) &&
-			ensure(InProgram.ConstantUInt32Lists.IsValidIndex(Args.GameplayTags.ListIndex)))
-		{
-			TagCount = InProgram.ConstantUInt32Lists.IsValidIndex(Args.GameplayTags.ListIndex) 
-					? InProgram.ConstantUInt32Lists[Args.GameplayTags.ListIndex].Num() 
-					: 0;
-		}
+				const int32 SkeletonCount = Args.SkeletonId != UE::Mutable::Private::PASSTHROUGH_ID_INVALID ? 1 : 0;
+				const int32 PhysicsAssetCount = Args.PhysicsAssetId != UE::Mutable::Private::PASSTHROUGH_ID_INVALID ? 1 : 0;
 
-		int32 ResourceCount = EnumHasAnyFlags(Args.Flags, OpEnumFlags::HasAnimationSlots) ? 1 : 0;
-		if (EnumHasAnyFlags(Args.Flags, OpEnumFlags::IsAnimationSlotList) &&
-			ensure(InProgram.ConstantUInt32Lists.IsValidIndex(Args.AnimSlotNames.ListIndex)))
-		{
-			ResourceCount = InProgram.ConstantUInt32Lists.IsValidIndex(Args.AnimSlotNames.ListIndex) 
-					? InProgram.ConstantUInt32Lists[Args.AnimSlotNames.ListIndex].Num() 
-					: 0;
-		}
+				int32 AdditionalPhysicsAssetCount = EnumHasAnyFlags(Args.Flags, OpEnumFlags::HasAdditionalPhysicsAsset) ? 1 : 0;
+				if (EnumHasAnyFlags(Args.Flags, OpEnumFlags::IsAdditionalPhysicsAssetList) &&
+					ensure(InProgram.ConstantUInt32Lists.IsValidIndex(Args.AdditionalPhysicsAssetsIds.ListAddress)))
+				{
+					AdditionalPhysicsAssetCount = InProgram.ConstantUInt32Lists[Args.AdditionalPhysicsAssetsIds.ListAddress].Num();
+				}
+			
+				int32 AssetUserDataCount = EnumHasAnyFlags(Args.Flags, OpEnumFlags::HasAssetUserData) ? 1 : 0;
+				if (EnumHasAnyFlags(Args.Flags, OpEnumFlags::IsAssetUserDataList) &&
+					ensure(InProgram.ConstantUInt32Lists.IsValidIndex(Args.AssetUserDataIds.ListAddress)))
+				{
+					AssetUserDataCount = InProgram.ConstantUInt32Lists[Args.AssetUserDataIds.ListAddress].Num();
+				}
 
-		int32 RealTimeMorphNamesCount = EnumHasAnyFlags(Args.Flags, OpEnumFlags::HasRealTimeMorphNames) ? 1 : 0;
-		if (EnumHasAnyFlags(Args.Flags, OpEnumFlags::IsRealTimeMorphNamesList) &&
-			ensure(InProgram.ConstantUInt32Lists.IsValidIndex(Args.RealTimeMorphNames.ListAddress)))
-		{
-			RealTimeMorphNamesCount = InProgram.ConstantUInt32Lists.IsValidIndex(Args.RealTimeMorphNames.ListAddress) 
-					? InProgram.ConstantUInt32Lists[Args.RealTimeMorphNames.ListAddress].Num() 
-					: 0;
-		}
+				FString Caption = FString::Printf(TEXT("add %d Tags, %d Resources, %d Skeleton, %d PhysicsAsset, %d RealTimeMorphs, %d Additional PhysicsAssets, and %d AssetUserData to"),
+					TagCount, ResourceCount, SkeletonCount, PhysicsAssetCount, RealTimeMorphNamesCount, AdditionalPhysicsAssetCount, AssetUserDataCount);
+				AddOpFunc(Address, Args.Source, Caption, ChildIndex);
 
-		const int32 SkeletonCount = Args.SkeletonId != UE::Mutable::Private::PASSTHROUGH_ID_INVALID ? 1 : 0;
-		const int32 PhysicsAssetCount = Args.PhysicsAssetId != UE::Mutable::Private::PASSTHROUGH_ID_INVALID ? 1 : 0;
+				break;
+			}
 
-		int32 AdditionalPhysicsAssetCount = EnumHasAnyFlags(Args.Flags, OpEnumFlags::HasAdditionalPhysicsAsset) ? 1 : 0;
-		if (EnumHasAnyFlags(Args.Flags, OpEnumFlags::IsAdditionalPhysicsAssetList) &&
-			ensure(InProgram.ConstantUInt32Lists.IsValidIndex(Args.AdditionalPhysicsAssetsIds.ListAddress)))
-		{
-			AdditionalPhysicsAssetCount = InProgram.ConstantUInt32Lists[Args.AdditionalPhysicsAssetsIds.ListAddress].Num();
-		}
+		case UE::Mutable::Private::EOpType::ME_APPLYLAYOUT:
+			{
+				uint32 ChildIndex = 0;
 			
-		int32 AssetUserDataCount = EnumHasAnyFlags(Args.Flags, OpEnumFlags::HasAssetUserData) ? 1 : 0;
-		if (EnumHasAnyFlags(Args.Flags, OpEnumFlags::IsAssetUserDataList) &&
-			ensure(InProgram.ConstantUInt32Lists.IsValidIndex(Args.AssetUserDataIds.ListAddress)))
-		{
-			AssetUserDataCount = InProgram.ConstantUInt32Lists[Args.AssetUserDataIds.ListAddress].Num();
-		}
-
-		FString Caption = FString::Printf(TEXT("add %d Tags, %d Resources, %d Skeleton, %d PhysicsAsset, %d RealTimeMorphs, %d Additional PhysicsAssets, and %d AssetUserData to"),
-			TagCount, ResourceCount, SkeletonCount, PhysicsAssetCount, RealTimeMorphNamesCount, AdditionalPhysicsAssetCount, AssetUserDataCount);
-		AddOpFunc(Args.Source, Caption);
-
-		break;
-	}
+				UE::Mutable::Private::FOperation::MeshApplyLayoutArgs Args = InProgram.GetOpArgs<UE::Mutable::Private::FOperation::MeshApplyLayoutArgs>(Address);
+				AddOpFunc(Address, Args.Layout, TEXT("layout "), ChildIndex);
+				AddOpFunc(Address, Args.Mesh, TEXT("mesh "), ChildIndex);
+				break;
+			}
 
-	case UE::Mutable::Private::EOpType::ME_APPLYLAYOUT:
-	{
-		UE::Mutable::Private::FOperation::MeshApplyLayoutArgs Args = InProgram.GetOpArgs<UE::Mutable::Private::FOperation::MeshApplyLayoutArgs>(InParentAddress);
-		AddOpFunc(Args.Layout, TEXT("layout "));
-		AddOpFunc(Args.Mesh, TEXT("mesh "));
-		break;
-	}
+		case UE::Mutable::Private::EOpType::ME_PREPARELAYOUT:
+			{
+				uint32 ChildIndex = 0;
+			
+				UE::Mutable::Private::FOperation::MeshPrepareLayoutArgs Args = InProgram.GetOpArgs<UE::Mutable::Private::FOperation::MeshPrepareLayoutArgs>(Address);
+				AddOpFunc(Address, Args.Layout, TEXT("layout "), ChildIndex);
+				AddOpFunc(Address, Args.Mesh, TEXT("mesh "), ChildIndex);
+				break;
+			}
 
-	case UE::Mutable::Private::EOpType::ME_PREPARELAYOUT:
-	{
-		UE::Mutable::Private::FOperation::MeshPrepareLayoutArgs Args = InProgram.GetOpArgs<UE::Mutable::Private::FOperation::MeshPrepareLayoutArgs>(InParentAddress);
-		AddOpFunc(Args.Layout, TEXT("layout "));
-		AddOpFunc(Args.Mesh, TEXT("mesh "));
-		break;
-	}
+		case UE::Mutable::Private::EOpType::ME_DIFFERENCE:
+			{
+				uint32 ChildIndex = 0;
 
-	case UE::Mutable::Private::EOpType::ME_DIFFERENCE:
-	{
-		const uint8* data = InProgram.GetOpArgsPointer(InParentAddress);
+				const uint8* data = InProgram.GetOpArgsPointer(Address);
 
-		UE::Mutable::Private::FOperation::ADDRESS BaseAt = 0;
-		FMemory::Memcpy(&BaseAt, data, sizeof(UE::Mutable::Private::FOperation::ADDRESS));
-		AddOpFunc(BaseAt, TEXT("base "));
-		data += sizeof(UE::Mutable::Private::FOperation::ADDRESS);
+				UE::Mutable::Private::FOperation::ADDRESS BaseAt = 0;
+				FMemory::Memcpy(&BaseAt, data, sizeof(UE::Mutable::Private::FOperation::ADDRESS));
+				AddOpFunc(Address, BaseAt, TEXT("base "), ChildIndex);
+				data += sizeof(UE::Mutable::Private::FOperation::ADDRESS);
 
-		UE::Mutable::Private::FOperation::ADDRESS TargetAt = 0;
-		FMemory::Memcpy(&TargetAt, data, sizeof(UE::Mutable::Private::FOperation::ADDRESS));
-		AddOpFunc(TargetAt, TEXT("target "));
-		data += sizeof(UE::Mutable::Private::FOperation::ADDRESS);
+				UE::Mutable::Private::FOperation::ADDRESS TargetAt = 0;
+				FMemory::Memcpy(&TargetAt, data, sizeof(UE::Mutable::Private::FOperation::ADDRESS));
+				AddOpFunc(Address, TargetAt, TEXT("target "), ChildIndex);
+				data += sizeof(UE::Mutable::Private::FOperation::ADDRESS);
 		
-		break;
-	}
+				break;
+			}
 		
-	case UE::Mutable::Private::EOpType::ME_MORPH:
-	{
-		const UE::Mutable::Private::FOperation::MeshMorphArgs Args = InProgram.GetOpArgs<UE::Mutable::Private::FOperation::MeshMorphArgs>(InParentAddress);
+		case UE::Mutable::Private::EOpType::ME_MORPH:
+			{
+				uint32 ChildIndex = 0;
 
-		FString MorphName = InProgram.ConstantStrings.IsValidIndex(Args.Name) 
-			? InProgram.ConstantStrings[Args.Name]
-			: FString(TEXT("[Name Not Found]"));
+				const UE::Mutable::Private::FOperation::MeshMorphArgs Args = InProgram.GetOpArgs<UE::Mutable::Private::FOperation::MeshMorphArgs>(Address);
+
+				FString MorphName = InProgram.ConstantStrings.IsValidIndex(Args.Name) 
+					? InProgram.ConstantStrings[Args.Name]
+					: FString(TEXT("[Name Not Found]"), ChildIndex);
 		
-		AddOpFunc(Args.Factor, TEXT("Factor "));
-		AddOpFunc(Args.Base, FString::Printf(TEXT("%s : Base "), *MorphName));
+				AddOpFunc(Address, Args.Factor, TEXT("Factor "), ChildIndex);
+				AddOpFunc(Address, Args.Base, FString::Printf(TEXT("%s : Base "), *MorphName), ChildIndex);
 	
-		break;
-	}
+				break;
+			}
 		
-	case UE::Mutable::Private::EOpType::SC_CURVE:
-	{
-		UE::Mutable::Private::FOperation::ScalarCurveArgs Args = InProgram.GetOpArgs<UE::Mutable::Private::FOperation::ScalarCurveArgs>(InParentAddress);
-		AddOpFunc(Args.time, TEXT("Time "));
-		break;
-	}
+		case UE::Mutable::Private::EOpType::SC_CURVE:
+			{
+				uint32 ChildIndex = 0;
+
+				UE::Mutable::Private::FOperation::ScalarCurveArgs Args = InProgram.GetOpArgs<UE::Mutable::Private::FOperation::ScalarCurveArgs>(Address);
+				AddOpFunc(Address, Args.time, TEXT("Time "), ChildIndex);
+				break;
+			}
 		
-	case UE::Mutable::Private::EOpType::CO_SAMPLEIMAGE:
-	{
-		UE::Mutable::Private::FOperation::ColourSampleImageArgs Args = InProgram.GetOpArgs<UE::Mutable::Private::FOperation::ColourSampleImageArgs>(InParentAddress);
-		AddOpFunc(Args.Image, TEXT("image "));
-		AddOpFunc(Args.X, TEXT("x "));
-		AddOpFunc(Args.Y, TEXT("y "));
-		break;
-	}
+		case UE::Mutable::Private::EOpType::CO_SAMPLEIMAGE:
+			{
+				uint32 ChildIndex = 0;
+			
+				UE::Mutable::Private::FOperation::ColourSampleImageArgs Args = InProgram.GetOpArgs<UE::Mutable::Private::FOperation::ColourSampleImageArgs>(Address);
+				AddOpFunc(Address, Args.Image, TEXT("image "), ChildIndex);
+				AddOpFunc(Address, Args.X, TEXT("x "), ChildIndex);
+				AddOpFunc(Address, Args.Y, TEXT("y "), ChildIndex);
+				break;
+			}
 
-	case UE::Mutable::Private::EOpType::CO_LINEARTOSRGB:
-	{
-		UE::Mutable::Private::FOperation::ColorArgs Args = InProgram.GetOpArgs<UE::Mutable::Private::FOperation::ColorArgs>(InParentAddress);
-		AddOpFunc(Args.Color, TEXT("Color "));
-		break;
-	}
+		case UE::Mutable::Private::EOpType::CO_LINEARTOSRGB:
+			{
+				uint32 ChildIndex = 0;
+
+				UE::Mutable::Private::FOperation::ColorArgs Args = InProgram.GetOpArgs<UE::Mutable::Private::FOperation::ColorArgs>(Address);
+				AddOpFunc(Address, Args.Color, TEXT("Color "), ChildIndex);
+				break;
+			}
 		
-	case UE::Mutable::Private::EOpType::IM_RESIZELIKE:
-	{
-		UE::Mutable::Private::FOperation::ImageResizeLikeArgs Args = InProgram.GetOpArgs<UE::Mutable::Private::FOperation::ImageResizeLikeArgs>(InParentAddress);
-		AddOpFunc(Args.Source, TEXT("src "));
-		AddOpFunc(Args.SizeSource, TEXT("sizeSrc "));
-		break;
-	}
+		case UE::Mutable::Private::EOpType::IM_RESIZELIKE:
+			{
+				uint32 ChildIndex = 0;
 
-	case UE::Mutable::Private::EOpType::IN_ADDMESH:
-	case UE::Mutable::Private::EOpType::IN_ADDSTRING:
-	case UE::Mutable::Private::EOpType::IN_ADDSURFACE:
-	case UE::Mutable::Private::EOpType::IN_ADDSKELETALMESH:
-	{
-		UE::Mutable::Private::FOperation::InstanceAddArgs Args = InProgram.GetOpArgs<UE::Mutable::Private::FOperation::InstanceAddArgs>(InParentAddress);
-		AddOpFunc(Args.instance, TEXT("instance "));
-		AddOpFunc(Args.value, TEXT("value "));
-		break;
-	}
+				UE::Mutable::Private::FOperation::ImageResizeLikeArgs Args = InProgram.GetOpArgs<UE::Mutable::Private::FOperation::ImageResizeLikeArgs>(Address);
+				AddOpFunc(Address, Args.Source, TEXT("src "), ChildIndex);
+				AddOpFunc(Address, Args.SizeSource, TEXT("sizeSrc "), ChildIndex);
+				break;
+			}
 		
-	case UE::Mutable::Private::EOpType::IN_ADDCOMPONENT:
-	{
-		UE::Mutable::Private::FOperation::InstanceAddComponentArgs Args = InProgram.GetOpArgs<UE::Mutable::Private::FOperation::InstanceAddComponentArgs>(InParentAddress);
-		AddOpFunc(Args.Instance, TEXT("instance "));
-		AddOpFunc(Args.Value, TEXT("value "));
-		break;
-	}
+		case UE::Mutable::Private::EOpType::IN_ADDMESH:
+		case UE::Mutable::Private::EOpType::IN_ADDSTRING:
+		case UE::Mutable::Private::EOpType::IN_ADDSURFACE:
+		case UE::Mutable::Private::EOpType::IN_ADDSKELETALMESH:
+			{
+				uint32 ChildIndex = 0;
+
+				UE::Mutable::Private::FOperation::InstanceAddArgs Args = InProgram.GetOpArgs<UE::Mutable::Private::FOperation::InstanceAddArgs>(Address);
+				AddOpFunc(Address, Args.instance, TEXT("instance "), ChildIndex);
+				AddOpFunc(Address, Args.value, TEXT("value "), ChildIndex);
+				break;
+			}
 		
-	case UE::Mutable::Private::EOpType::IN_ADDOVERLAYMATERIAL:
-	{
-		UE::Mutable::Private::FOperation::InstanceAddOverlayMaterialArgs Args = InProgram.GetOpArgs<UE::Mutable::Private::FOperation::InstanceAddOverlayMaterialArgs>(InParentAddress);
-		AddOpFunc(Args.Instance, TEXT("instance "));
-		AddOpFunc(Args.Material, TEXT("value "));
-		break;
-	}
+		case UE::Mutable::Private::EOpType::IN_ADDCOMPONENT:
+			{
+				uint32 ChildIndex = 0;
+	
+				UE::Mutable::Private::FOperation::InstanceAddComponentArgs Args = InProgram.GetOpArgs<UE::Mutable::Private::FOperation::InstanceAddComponentArgs>(Address);
+				AddOpFunc(Address, Args.Instance, TEXT("instance "), ChildIndex);
+				AddOpFunc(Address, Args.Value, TEXT("value "), ChildIndex);
+				break;
+			}
+		
+		case UE::Mutable::Private::EOpType::IN_ADDOVERLAYMATERIAL:
+			{
+				uint32 ChildIndex = 0;
 
-	case UE::Mutable::Private::EOpType::IN_ADDOVERRIDEMATERIAL:
-	{
-		UE::Mutable::Private::FOperation::InstanceAddOverrideMaterialArgs Args = InProgram.GetOpArgs<UE::Mutable::Private::FOperation::InstanceAddOverrideMaterialArgs>(InParentAddress);
-		AddOpFunc(Args.Instance, TEXT("instance "));
-		AddOpFunc(Args.Material, TEXT("value "));
-		break;
-	}
+				UE::Mutable::Private::FOperation::InstanceAddOverlayMaterialArgs Args = InProgram.GetOpArgs<UE::Mutable::Private::FOperation::InstanceAddOverlayMaterialArgs>(Address);
+				AddOpFunc(Address, Args.Instance, TEXT("instance "), ChildIndex);
+				AddOpFunc(Address, Args.Material, TEXT("value "), ChildIndex);
+				break;
+			}
+
+		case UE::Mutable::Private::EOpType::IN_ADDOVERRIDEMATERIAL:
+			{
+				uint32 ChildIndex = 0;
+
+				UE::Mutable::Private::FOperation::InstanceAddOverrideMaterialArgs Args = InProgram.GetOpArgs<UE::Mutable::Private::FOperation::InstanceAddOverrideMaterialArgs>(Address);
+				AddOpFunc(Address, Args.Instance, TEXT("instance "), ChildIndex);
+				AddOpFunc(Address, Args.Material, TEXT("value "), ChildIndex);
+				break;
+			}
 	
 	case UE::Mutable::Private::EOpType::IN_ADDMATERIAL:
-	{
-		UE::Mutable::Private::FOperation::InstanceAddMaterialArgs Args = InProgram.GetOpArgs<UE::Mutable::Private::FOperation::InstanceAddMaterialArgs>(InParentAddress);
-		AddOpFunc(Args.Instance, TEXT("instance "));
-		AddOpFunc(Args.Material, TEXT("material "));
-		break;
-	}
+		{
+			uint32 ChildIndex = 0;
+
+			UE::Mutable::Private::FOperation::InstanceAddMaterialArgs Args = InProgram.GetOpArgs<UE::Mutable::Private::FOperation::InstanceAddMaterialArgs>(Address);
+			AddOpFunc(Address, Args.Instance, TEXT("instance "), ChildIndex);
+			AddOpFunc(Address, Args.Material, TEXT("material "), ChildIndex);
+			break;
+		}
 	
 	case UE::Mutable::Private::EOpType::IM_COMPOSE:
-	{
-		UE::Mutable::Private::FOperation::ImageComposeArgs Args = InProgram.GetOpArgs<UE::Mutable::Private::FOperation::ImageComposeArgs>(InParentAddress);
-		AddOpFunc(Args.layout, TEXT("Layout "));
-		AddOpFunc(Args.base, TEXT("Base "));
-		AddOpFunc(Args.blockImage, TEXT("BlockImage "));
-		AddOpFunc(Args.mask, TEXT("Mask "));
-		break;
-	}
+		{
+			uint32 ChildIndex = 0;
 
-	case UE::Mutable::Private::EOpType::IM_MULTICOMPOSE:
-	{
-		UE::Mutable::Private::FOperation::ImageMultiComposeArgs Args = InProgram.GetOpArgs<UE::Mutable::Private::FOperation::ImageMultiComposeArgs>(InParentAddress);
-		AddOpFunc(Args.Layout, TEXT("Layout "));
-		AddOpFunc(Args.Base, TEXT("Base "));
-		AddOpFunc(Args.SourceLayout, TEXT("SourceLayout"));
-		AddOpFunc(Args.SourceImage, TEXT("SourceImage"));
+			UE::Mutable::Private::FOperation::ImageComposeArgs Args = InProgram.GetOpArgs<UE::Mutable::Private::FOperation::ImageComposeArgs>(Address);
+			AddOpFunc(Address, Args.layout, TEXT("Layout "), ChildIndex);
+			AddOpFunc(Address, Args.base, TEXT("Base "), ChildIndex);
+			AddOpFunc(Address, Args.blockImage, TEXT("BlockImage "), ChildIndex);
+			AddOpFunc(Address, Args.mask, TEXT("Mask "), ChildIndex);
+			break;
+		}
+
+		case UE::Mutable::Private::EOpType::IM_MULTICOMPOSE:
+			{
+				uint32 ChildIndex = 0;
+
+				UE::Mutable::Private::FOperation::ImageMultiComposeArgs Args = InProgram.GetOpArgs<UE::Mutable::Private::FOperation::ImageMultiComposeArgs>(Address);
+				AddOpFunc(Address, Args.Layout, TEXT("Layout "), ChildIndex);
+				AddOpFunc(Address, Args.Base, TEXT("Base "), ChildIndex);
+				AddOpFunc(Address, Args.SourceLayout, TEXT("SourceLayout"), ChildIndex);
+				AddOpFunc(Address, Args.SourceImage, TEXT("SourceImage"), ChildIndex);
 		
-		break;
-	}
+				break;
+			}
 		
-	case UE::Mutable::Private::EOpType::IM_INTERPOLATE:
-	{
-		UE::Mutable::Private::FOperation::ImageInterpolateArgs Args = InProgram.GetOpArgs<UE::Mutable::Private::FOperation::ImageInterpolateArgs>(InParentAddress);
-		AddOpFunc(Args.Factor, TEXT("factor "));
+		case UE::Mutable::Private::EOpType::IM_INTERPOLATE:
+			{
+				uint32 ChildIndex = 0;
+
+				UE::Mutable::Private::FOperation::ImageInterpolateArgs Args = InProgram.GetOpArgs<UE::Mutable::Private::FOperation::ImageInterpolateArgs>(Address);
+				AddOpFunc(Address, Args.Factor, TEXT("factor "), ChildIndex);
 			
-		int32 TargetIndex = 0;
-		for (const UE::Mutable::Private::FOperation::ADDRESS& Operation : Args.Targets)
-		{
-			AddOpFunc(Operation, FString::Printf(TEXT("target %i "), TargetIndex++));
-		}
-		break;
-	}
+				int32 TargetIndex = 0;
+				for (const UE::Mutable::Private::FOperation::ADDRESS& Operation : Args.Targets)
+				{
+					AddOpFunc(Address, Operation, FString::Printf(TEXT("target %i "), TargetIndex++), ChildIndex);
+				}
+				break;
+			}
 		
-	case UE::Mutable::Private::EOpType::IM_SATURATE:
-	{
-		UE::Mutable::Private::FOperation::ImageSaturateArgs Args = InProgram.GetOpArgs<UE::Mutable::Private::FOperation::ImageSaturateArgs>(InParentAddress);
-		AddOpFunc(Args.Base, TEXT("base "));
-		AddOpFunc(Args.Factor, TEXT("factor "));
-		break;
-	}
+		case UE::Mutable::Private::EOpType::IM_SATURATE:
+			{
+				uint32 ChildIndex = 0;
+
+				UE::Mutable::Private::FOperation::ImageSaturateArgs Args = InProgram.GetOpArgs<UE::Mutable::Private::FOperation::ImageSaturateArgs>(Address);
+				AddOpFunc(Address, Args.Base, TEXT("base "), ChildIndex);
+				AddOpFunc(Address, Args.Factor, TEXT("factor "), ChildIndex);
+				break;
+			}
 		
-	case UE::Mutable::Private::EOpType::IM_COLOURMAP:
-	{
-		UE::Mutable::Private::FOperation::ImageColourMapArgs Args = InProgram.GetOpArgs<UE::Mutable::Private::FOperation::ImageColourMapArgs>(InParentAddress);
-		AddOpFunc(Args.Base, TEXT("base "));
-		AddOpFunc(Args.Mask, TEXT("mask "));
-		AddOpFunc(Args.Map, TEXT("map "));
-		break;
-	}
+		case UE::Mutable::Private::EOpType::IM_COLOURMAP:
+			{
+				uint32 ChildIndex = 0;
+
+				UE::Mutable::Private::FOperation::ImageColourMapArgs Args = InProgram.GetOpArgs<UE::Mutable::Private::FOperation::ImageColourMapArgs>(Address);
+				AddOpFunc(Address, Args.Base, TEXT("base "), ChildIndex);
+				AddOpFunc(Address, Args.Mask, TEXT("mask "), ChildIndex);
+				AddOpFunc(Address, Args.Map, TEXT("map "), ChildIndex);
+				break;
+			}
 		
-	case UE::Mutable::Private::EOpType::IM_BINARISE:
-	{
-		UE::Mutable::Private::FOperation::ImageBinariseArgs Args = InProgram.GetOpArgs<UE::Mutable::Private::FOperation::ImageBinariseArgs>(InParentAddress);
-		AddOpFunc(Args.Base, TEXT("base "));
-		AddOpFunc(Args.Threshold, TEXT("threshold "));
-		break;
-	}
+		case UE::Mutable::Private::EOpType::IM_BINARISE:
+			{
+				uint32 ChildIndex = 0;
+
+				UE::Mutable::Private::FOperation::ImageBinariseArgs Args = InProgram.GetOpArgs<UE::Mutable::Private::FOperation::ImageBinariseArgs>(Address);
+				AddOpFunc(Address, Args.Base, TEXT("base "), ChildIndex);
+				AddOpFunc(Address, Args.Threshold, TEXT("threshold "), ChildIndex);
+				break;
+			}
 		
-	case UE::Mutable::Private::EOpType::IM_PATCH:
-	{
-		UE::Mutable::Private::FOperation::ImagePatchArgs Args = InProgram.GetOpArgs<UE::Mutable::Private::FOperation::ImagePatchArgs>(InParentAddress);
-		AddOpFunc(Args.base, TEXT("base "));
-		AddOpFunc(Args.patch, TEXT("patch "));
-		break;
-	}
+		case UE::Mutable::Private::EOpType::IM_PATCH:
+			{
+				uint32 ChildIndex = 0;
+	
+				UE::Mutable::Private::FOperation::ImagePatchArgs Args = InProgram.GetOpArgs<UE::Mutable::Private::FOperation::ImagePatchArgs>(Address);
+				AddOpFunc(Address, Args.base, TEXT("base "), ChildIndex);
+				AddOpFunc(Address, Args.patch, TEXT("patch "), ChildIndex);
+				break;
+			}
 		
-	case UE::Mutable::Private::EOpType::IM_RASTERMESH:
-	{
-		UE::Mutable::Private::FOperation::ImageRasterMeshArgs Args = InProgram.GetOpArgs<UE::Mutable::Private::FOperation::ImageRasterMeshArgs>(InParentAddress);
-		AddOpFunc(Args.mesh, TEXT("mesh "));
-		AddOpFunc(Args.image, TEXT("image "));
-		AddOpFunc(Args.mask, TEXT("mask "));
-		AddOpFunc(Args.angleFadeProperties, TEXT("angleFadeProperties "));
-		AddOpFunc(Args.projector, TEXT("projector "));
-		break;
-	}
+		case UE::Mutable::Private::EOpType::IM_RASTERMESH:
+			{
+				uint32 ChildIndex = 0;
 		
-	case UE::Mutable::Private::EOpType::IM_DISPLACE:
-	{
-		UE::Mutable::Private::FOperation::ImageDisplaceArgs Args = InProgram.GetOpArgs<UE::Mutable::Private::FOperation::ImageDisplaceArgs>(InParentAddress);
-		AddOpFunc(Args.Source, TEXT("src "));
-		AddOpFunc(Args.DisplacementMap, TEXT("displacementMap "));
-		break;
-	}
+				UE::Mutable::Private::FOperation::ImageRasterMeshArgs Args = InProgram.GetOpArgs<UE::Mutable::Private::FOperation::ImageRasterMeshArgs>(Address);
+				AddOpFunc(Address, Args.mesh, TEXT("mesh "), ChildIndex);
+				AddOpFunc(Address, Args.image, TEXT("image "), ChildIndex);
+				AddOpFunc(Address, Args.mask, TEXT("mask "), ChildIndex);
+				AddOpFunc(Address, Args.angleFadeProperties, TEXT("angleFadeProperties "), ChildIndex);
+				AddOpFunc(Address, Args.projector, TEXT("projector "), ChildIndex);
+				break;
+			}
 		
-	case UE::Mutable::Private::EOpType::IM_NORMALCOMPOSITE:
-	{
-		UE::Mutable::Private::FOperation::ImageNormalCompositeArgs Args = InProgram.GetOpArgs<UE::Mutable::Private::FOperation::ImageNormalCompositeArgs>(InParentAddress);
-		AddOpFunc(Args.base, TEXT("base "));
-		AddOpFunc(Args.normal, TEXT("normal "));
-		break;
-	}
+		case UE::Mutable::Private::EOpType::IM_DISPLACE:
+			{
+				uint32 ChildIndex = 0;
+
+				UE::Mutable::Private::FOperation::ImageDisplaceArgs Args = InProgram.GetOpArgs<UE::Mutable::Private::FOperation::ImageDisplaceArgs>(Address);
+				AddOpFunc(Address, Args.Source, TEXT("src "), ChildIndex);
+				AddOpFunc(Address, Args.DisplacementMap, TEXT("displacementMap "), ChildIndex);
+				break;
+			}
 		
-	case UE::Mutable::Private::EOpType::IM_TRANSFORM:
-	{
-		UE::Mutable::Private::FOperation::ImageTransformArgs Args = InProgram.GetOpArgs<UE::Mutable::Private::FOperation::ImageTransformArgs>(InParentAddress);
-		AddOpFunc(Args.Base, TEXT("base "));
-		AddOpFunc(Args.OffsetX, TEXT("offsetX "));
-		AddOpFunc(Args.OffsetY, TEXT("offsetY "));
-		AddOpFunc(Args.ScaleX, TEXT("scaleX "));
-		AddOpFunc(Args.ScaleY, TEXT("scaleY "));
-		AddOpFunc(Args.Rotation, TEXT("rotation "));
-		break;
-	}
+		case UE::Mutable::Private::EOpType::IM_NORMALCOMPOSITE:
+			{
+				uint32 ChildIndex = 0;
 
-	case UE::Mutable::Private::EOpType::SC_MATERIAL_BREAK:
-	case UE::Mutable::Private::EOpType::CO_MATERIAL_BREAK:
-	case UE::Mutable::Private::EOpType::IM_MATERIAL_BREAK:
-	{
-		UE::Mutable::Private::FOperation::MaterialBreakArgs Args = InProgram.GetOpArgs<UE::Mutable::Private::FOperation::MaterialBreakArgs>(InParentAddress);
-		AddOpFunc(Args.Material, TEXT("Material "));
-		break;
-	}
+				UE::Mutable::Private::FOperation::ImageNormalCompositeArgs Args = InProgram.GetOpArgs<UE::Mutable::Private::FOperation::ImageNormalCompositeArgs>(Address);
+				AddOpFunc(Address, Args.base, TEXT("base "), ChildIndex);
+				AddOpFunc(Address, Args.normal, TEXT("normal "), ChildIndex);
+				break;
+			}
+		
+		case UE::Mutable::Private::EOpType::IM_TRANSFORM:
+			{
+				uint32 ChildIndex = 0;
+
+				UE::Mutable::Private::FOperation::ImageTransformArgs Args = InProgram.GetOpArgs<UE::Mutable::Private::FOperation::ImageTransformArgs>(Address);
+				AddOpFunc(Address, Args.Base, TEXT("base "), ChildIndex);
+				AddOpFunc(Address, Args.OffsetX, TEXT("offsetX "), ChildIndex);
+				AddOpFunc(Address, Args.OffsetY, TEXT("offsetY "), ChildIndex);
+				AddOpFunc(Address, Args.ScaleX, TEXT("scaleX "), ChildIndex);
+				AddOpFunc(Address, Args.ScaleY, TEXT("scaleY "), ChildIndex);
+				AddOpFunc(Address, Args.Rotation, TEXT("rotation "), ChildIndex);
+				break;
+			}
+
+		case UE::Mutable::Private::EOpType::SC_MATERIAL_BREAK:
+		case UE::Mutable::Private::EOpType::CO_MATERIAL_BREAK:
+		case UE::Mutable::Private::EOpType::IM_MATERIAL_BREAK:
+			{
+				uint32 ChildIndex = 0;
+
+				UE::Mutable::Private::FOperation::MaterialBreakArgs Args = InProgram.GetOpArgs<UE::Mutable::Private::FOperation::MaterialBreakArgs>(Address);
+				AddOpFunc(Address, Args.Material, TEXT("Material "), ChildIndex);
+				break;
+			}
 
-	//TODO(Max) Add MI_Parameter and MI_CONSTANT
+			//TODO(Max) Add MI_Parameter and MI_CONSTANT
 
-	// Add here more operation types to define how they are exposed in the tree (set a Caption)
+			// Add here more operation types to define how they are exposed in the tree (set a Caption)
 		
-	default:
-		{
-			// Generic list of child operations
-			bUseGeneric = true;
-			break;
+		default:
+			{
+				// Generic list of child operations
+				bUseGeneric = true;
+				break;
+			}
 		}
-	}
 
-	if (bUseGeneric)
-	{
-		// Find children of the provided element without adding any extra string to better identify what each of them represents
-		UE::Mutable::Private::ForEachReference(InProgram, InParentAddress, [this, &InProgram, AddOpFunc](UE::Mutable::Private::FOperation::ADDRESS ChildAddress)
+		if (bUseGeneric)
+		{
+			// Find children of the provided element without adding any extra string to better identify what each of them represents
+			uint32 ChildIndex = 0;
+			UE::Mutable::Private::ForEachReference(InProgram, Address, [this, AddOpFunc, &ChildIndex, Address](UE::Mutable::Private::FOperation::ADDRESS ChildAddress)
 			{
-				AddOpFunc(ChildAddress,TEXT(""));
+				AddOpFunc(Address, ChildAddress, TEXT(""), ChildIndex);
 			});
-	}
-	else
-	{
-		// Validate in case there is a mismatch in the custom processing of children and the generic one, which would cause problems.
-		ChildIndex = 0;
-
-		auto ValidateOpFunc = [this, InParentAddress, ParentOperationType, &ChildIndex](UE::Mutable::Private::FOperation::ADDRESS ChildAddress)
+		}
+		else
 		{
-			const FItemCacheKey Key = { InParentAddress, ChildAddress, ChildIndex };
-			const TSharedPtr<FMutableCodeTreeElement>* CachedItem = ItemCache.Find(Key);
+			// Validate in case there is a mismatch in the custom processing of children and the generic one, which would cause problems.
+			uint32 ChildIndex = 0;
 
-			// If this check fails could mean that in the switch above one type of OP processes its children in one
-			// order but in "ForEachReference" that same operation processes the same children in another order.
-			check(CachedItem);
+			auto ValidateOpFunc = [this, Address, &ChildIndex, &InProgram](UE::Mutable::Private::FOperation::ADDRESS ChildAddress)
+				{
+					const FItemCacheKey Key = { Address, ChildAddress, ChildIndex };
+					const TSharedPtr<FMutableCodeTreeElement>* CachedItem = ItemCache.Find(Key);
+
+					volatile auto a = InProgram.GetOpType(ChildAddress);
+					volatile auto b = InProgram.GetOpType(Address);
+					
+					// If this check fails could mean that in the switch above one type of OP processes its children in one
+					// order but in "ForEachReference" that same operation processes the same children in another order.
+					check(CachedItem);
 				
-			++ChildIndex;
-		};
+					++ChildIndex;
+				};
 
-		UE::Mutable::Private::ForEachReference(InProgram, InParentAddress, [this, ValidateOpFunc](UE::Mutable::Private::FOperation::ADDRESS ChildAddress)
-		{
-			ValidateOpFunc(ChildAddress);
-		});
+			UE::Mutable::Private::ForEachReference(InProgram, Address, [this, ValidateOpFunc](UE::Mutable::Private::FOperation::ADDRESS ChildAddress)
+			{
+				ValidateOpFunc(ChildAddress);
+			});
+		}
 	}
 }
 
@@ -2109,14 +2184,7 @@ void SMutableCodeViewer::OnSelectionChanged(TSharedPtr<FMutableCodeTreeElement>
 		PrepareMaterialViewer();
 		break;
 	}
-
-	case UE::Mutable::Private::EDataType::Instance:
-	{
-		// Create or reuse the UI
-		PrepareInstanceViewer();
-		break;
-	}
-
+	
 	case UE::Mutable::Private::EDataType::Scalar:
 	{
 		// Create or reuse the UI
@@ -2835,15 +2903,7 @@ void SMutableCodeViewer::Tick(const FGeometry& AllottedGeometry, const double In
 		PreviewMeshViewer->SetMesh(MutableMesh);
 		break;
 	}
-
-	case UE::Mutable::Private::EDataType::Instance:
-	{
-		check(PreviewInstanceViewer);
-		UE::Mutable::Private::TManagedPtr<const UE::Mutable::Private::FInstance> MutableInstance = System->BuildInstance(TempLiveInstance, SelectedOperationAddress);
-		PreviewInstanceViewer->SetInstance(TempLiveInstance, MutableInstance, PreviewParameters, *System);
-		break;
-	}
-
+	
 	case UE::Mutable::Private::EDataType::Bool:
 	{
 		check(PreviewBoolViewer);
@@ -2979,17 +3039,6 @@ void SMutableCodeViewer::PrepareMaterialViewer()
 }
 
 
-void SMutableCodeViewer::PrepareInstanceViewer()
-{
-	if (!PreviewInstanceViewer)
-	{
-		PreviewInstanceViewer = SNew(SMutableInstanceViewer);
-	}
-
-	PreviewBorder->SetContent(PreviewInstanceViewer.ToSharedRef());
-}
-
-
 void SMutableCodeViewer::PrepareLayoutViewer()
 {
 	if (!PreviewLayoutViewer)
@@ -3299,11 +3348,11 @@ void FMutableCodeTreeElement::GenerateLabelText()
 	if (!Caption.IsEmpty())
 	{
 		Caption = Caption.TrimStartAndEnd();
-		MainLabel = FString::Printf(TEXT("%d : [ %s ] %s"), int32(MutableOperation), *Caption, *OpName);
+		MainLabel = FString::Printf(TEXT("%d: [%s] %s"), int32(MutableOperation), *Caption, *OpName);
 	}
 	else
 	{
-		MainLabel = FString::Printf(TEXT("%d : %s"), int32(MutableOperation), *OpName);
+		MainLabel = FString::Printf(TEXT("%d: %s"), int32(MutableOperation), *OpName);
 	}
 	
 
diff --git a/Engine/Plugins/Mutable/Source/CustomizableObjectEditor/Private/MuCOE/SMutableCodeViewer.h b/Engine/Plugins/Mutable/Source/CustomizableObjectEditor/Private/MuCOE/SMutableCodeViewer.h
--- a/Engine/Plugins/Mutable/Source/CustomizableObjectEditor/Private/MuCOE/SMutableCodeViewer.h
+++ b/Engine/Plugins/Mutable/Source/CustomizableObjectEditor/Private/MuCOE/SMutableCodeViewer.h
@@ -115,10 +115,7 @@ private:
 
 	/** Widget used to show the preview of image operation results. Once created it is reused to preserve the settings. */
 	TSharedPtr<SMutableImageViewer> PreviewImageViewer;
-
-	/** Widget used to show a preview of an instance. */
-	TSharedPtr<SMutableInstanceViewer> PreviewInstanceViewer;
-
+	
 	/** Widget used to show a preview of a mesh and the metadata it holds */
 	TSharedPtr<SMutableMeshViewer> PreviewMeshViewer;
 
@@ -202,7 +199,7 @@ private:
 	void GenerateAllTreeElements();
 
 	/** Support GenerateAllTreeElements by generating the elements recursively */
-	void GenerateElementRecursive( const int32& InStateIndex, UE::Mutable::Private::FOperation::ADDRESS InParentAddress, const UE::Mutable::Private::FProgram& InProgram);
+	void GenerateElement( const int32& InStateIndex, UE::Mutable::Private::FOperation::ADDRESS InParentAddress, const UE::Mutable::Private::FProgram& InProgram);
 
 	
 	/** The addresses of the root operations. Cached when this object gets loaded on the Construct method of this slate */
@@ -820,7 +817,6 @@ private:
 	void PrepareImageViewer();
 	void PrepareMeshViewer();
 	void PrepareMaterialViewer();
-	void PrepareInstanceViewer();
 	void PrepareLayoutViewer();
 	void PrepareProjectorViewer();
 };
diff --git a/Engine/Plugins/Mutable/Source/CustomizableObjectEditor/Private/MuCOE/SMutableInstanceViewer.cpp b/Engine/Plugins/Mutable/Source/CustomizableObjectEditor/Private/MuCOE/SMutableInstanceViewer.cpp
deleted file mode 100644
diff --git a/Engine/Plugins/Mutable/Source/CustomizableObjectEditor/Private/MuCOE/SMutableInstanceViewer.h b/Engine/Plugins/Mutable/Source/CustomizableObjectEditor/Private/MuCOE/SMutableInstanceViewer.h
deleted file mode 100644
diff --git a/Engine/Plugins/Mutable/Source/MutableRuntime/Internal/MuR/CodeVisitor.h b/Engine/Plugins/Mutable/Source/MutableRuntime/Internal/MuR/CodeVisitor.h
--- a/Engine/Plugins/Mutable/Source/MutableRuntime/Internal/MuR/CodeVisitor.h
+++ b/Engine/Plugins/Mutable/Source/MutableRuntime/Internal/MuR/CodeVisitor.h
@@ -258,6 +258,7 @@ namespace UE::Mutable::Private
             case EOpType::LA_CONDITIONAL:
             case EOpType::IN_CONDITIONAL:
 			case EOpType::ED_CONDITIONAL:
+			case EOpType::MI_CONDITIONAL:
             {
 				FOperation::ConditionalArgs Args = Program.GetOpArgs<FOperation::ConditionalArgs>(Address);
 
diff --git a/Engine/Plugins/Mutable/Source/MutableTools/Private/MuT/ASTOpConditional.cpp b/Engine/Plugins/Mutable/Source/MutableTools/Private/MuT/ASTOpConditional.cpp
--- a/Engine/Plugins/Mutable/Source/MutableTools/Private/MuT/ASTOpConditional.cpp
+++ b/Engine/Plugins/Mutable/Source/MutableTools/Private/MuT/ASTOpConditional.cpp
@@ -86,6 +86,7 @@ namespace UE::Mutable::Private
 		case EOpType::LA_CONDITIONAL:
 		case EOpType::IN_CONDITIONAL:
 		case EOpType::ED_CONDITIONAL:
+		case EOpType::MI_CONDITIONAL:
 			break;
 		default:
 			// Unexpected type
