[mutable] Fixed issue preventing the filling of the drag-and-drop pin node spawning options.

#rnx
#Synced-CL 51999473

diff --git a/Engine/Plugins/Mutable/Source/CustomizableObjectEditor/Private/MuCOE/Nodes/CustomizableObjectNodeComponentMesh.cpp b/Engine/Plugins/Mutable/Source/CustomizableObjectEditor/Private/MuCOE/Nodes/CustomizableObjectNodeComponentMesh.cpp
--- a/Engine/Plugins/Mutable/Source/CustomizableObjectEditor/Private/MuCOE/Nodes/CustomizableObjectNodeComponentMesh.cpp
+++ b/Engine/Plugins/Mutable/Source/CustomizableObjectEditor/Private/MuCOE/Nodes/CustomizableObjectNodeComponentMesh.cpp
@@ -244,7 +244,7 @@ bool UCustomizableObjectNodeComponentMesh::CanConnect(const UEdGraphPin* InOwned
 		return !bOutIsOtherNodeBlocklisted;
 	}
 
-	return false;
+	return true;
 }
 
 
