[Mutable] Fix crash when validating Customizable Objects.

#Synced-CL 51880525

diff --git a/Engine/Plugins/Mutable/Source/CustomizableObjectEditor/Private/MuCOE/CustomizableObjectEditorModule.cpp b/Engine/Plugins/Mutable/Source/CustomizableObjectEditor/Private/MuCOE/CustomizableObjectEditorModule.cpp
--- a/Engine/Plugins/Mutable/Source/CustomizableObjectEditor/Private/MuCOE/CustomizableObjectEditorModule.cpp
+++ b/Engine/Plugins/Mutable/Source/CustomizableObjectEditor/Private/MuCOE/CustomizableObjectEditorModule.cpp
@@ -1418,38 +1418,43 @@ void FCustomizableObjectEditorModule::Tick(bool bBlocking)
 	if (TSet<TStrongObjectPtr<UCustomizableObject>>::TIterator It = ObjectsPendingLoad.CreateIterator())
 	{
 		UCustomizableObject* Object = It->Get();
-		check(IsValid(Object));
-		
-		TSet<UCustomizableObject*> RelatedObjects;
-		const bool bPendingLoad = GraphTraversal::GetAllObjectsInGraphIncremental(Object, RelatedObjects);
-
-		if (!bPendingLoad)
+		if (IsValid(Object)) // Someone may have called MarkAsGarbage.
 		{
-			bool bHierarchyLoaded = true;
-			for (const UCustomizableObject* RelatedObject : RelatedObjects)
+			TSet<UCustomizableObject*> RelatedObjects;
+			const bool bPendingLoad = GraphTraversal::GetAllObjectsInGraphIncremental(Object, RelatedObjects);
+
+			if (!bPendingLoad)
 			{
-				if (RelatedObject && !RelatedObject->HasAnyFlags(EObjectFlags::RF_LoadCompleted))
+				bool bHierarchyLoaded = true;
+				for (const UCustomizableObject* RelatedObject : RelatedObjects)
 				{
-					bHierarchyLoaded = false;
-					break;
+					if (RelatedObject && !RelatedObject->HasAnyFlags(EObjectFlags::RF_LoadCompleted))
+					{
+						bHierarchyLoaded = false;
+						break;
+					}
 				}
-			}
 			
-			if (bHierarchyLoaded)
-			{	
-				// Traversing the graph upwards does not take into account Group Object Nodes. 
-				// A Child Object may reference a parent but not a Group Object. This can cause the loaded CO to not appear in the Related objects. Added manually to fix this.
-				RelatedObjects.Add(Object);
+				if (bHierarchyLoaded)
+				{	
+					// Traversing the graph upwards does not take into account Group Object Nodes. 
+					// A Child Object may reference a parent but not a Group Object. This can cause the loaded CO to not appear in the Related objects. Added manually to fix this.
+					RelatedObjects.Add(Object);
 				
-				for (UCustomizableObject* RelatedObject : RelatedObjects)
-				{
-					if (ObjectsPendingLoad.Remove(TStrongObjectPtr(RelatedObject)))
+					for (UCustomizableObject* RelatedObject : RelatedObjects)
 					{
-						OnUpstreamCOsLoaded(RelatedObject);
-						RelatedObject->GetPrivate()->LoadCompiledDataFromDisk();
+						if (ObjectsPendingLoad.Remove(TStrongObjectPtr(RelatedObject)))
+						{
+							OnUpstreamCOsLoaded(RelatedObject);
+							RelatedObject->GetPrivate()->LoadCompiledDataFromDisk();
+						}
 					}
 				}
-			}
+			}	
+		}
+		else
+		{
+			It.RemoveCurrent();
 		}
 	}
 	
