[Mutable] Optimize OpMeshRemove morphs.

- Only partialy edited batches data is modified, untouched batches or totally removed batches only modify batch headers.
- Now empty batches are removed.

#rnx
#Synced-CL 52018704

diff --git a/Engine/Plugins/Mutable/Source/MutableRuntime/Private/MuR/OpMeshRemove.cpp b/Engine/Plugins/Mutable/Source/MutableRuntime/Private/MuR/OpMeshRemove.cpp
--- a/Engine/Plugins/Mutable/Source/MutableRuntime/Private/MuR/OpMeshRemove.cpp
+++ b/Engine/Plugins/Mutable/Source/MutableRuntime/Private/MuR/OpMeshRemove.cpp
@@ -12,73 +12,137 @@
 
 namespace UE::Mutable::Private
 {
+	int32 RemoveBufferUnusedElems(uint8* BufferData, uint32 ElemSize, int32 NumElems, const TBitArray<>& UsedElems)
+	{
+    	if (ElemSize == 0)
+    	{
+    		return 0;
+    	}
+    	
+		int32 LastFreeElemIndex = 0;
+
+		for (int32 ElemIndex = UsedElems.FindFrom(true, 0); ElemIndex >= 0;)
+		{
+			const int32 UsedSpanBegin = ElemIndex;
+			ElemIndex = UsedElems.FindFrom(false, ElemIndex);
+			ElemIndex = ElemIndex >= 0 ? ElemIndex : NumElems;
+			const int32 UsedSpanEnd = ElemIndex;
+
+			ElemIndex = UsedElems.FindFrom(true, ElemIndex);
+
+			// Copy vertex buffer span.	
+			const int32 UsedSpanSize = UsedSpanEnd - UsedSpanBegin;
+            
+			if (LastFreeElemIndex != UsedSpanBegin)
+			{
+				FMemory::Memmove(
+						BufferData + LastFreeElemIndex*ElemSize,
+						BufferData + UsedSpanBegin*ElemSize,
+						UsedSpanSize*ElemSize);
+			}
+
+			LastFreeElemIndex += UsedSpanSize;
+		}
+
+		return LastFreeElemIndex;
+	}
+
 
 	void RemoveMorphBatchVertices(
 		UE::MorphTargetVertexCodec::FDeltaBatchHeader& InOutBatchHeader, 
 		TArrayView<uint32> DataView, 
-		//int32 BatchVertexRangeBegin, 
-		//int32 BatchVertexRangeEnd, 
+		int32 BatchVertexRangeBegin, 
+		int32 BatchVertexRangeEnd,
+		const TBitArray<>& UsedVertices,
 		TConstArrayView<int32> VertexRemap,
 		int32 NumResultVertices
 	)
 	{
-		UE::MorphTargetVertexCodec::FQuantizedDelta QuantizedDeltasReadBuffer[UE::MorphTargetVertexCodec::BatchSize];
+		check(VertexRemap.Num() == UsedVertices.Num());
 
-		TArrayView<uint32> BatchDataView = TArrayView<uint32>(
-				DataView.GetData() + (InOutBatchHeader.DataOffset/sizeof(uint32)), 
-				CalculateBatchDwords(InOutBatchHeader)); 
+		UE::MorphTargetVertexCodec::FQuantizedDelta QuantizedDeltasReadBuffer[UE::MorphTargetVertexCodec::BatchSize];
 
-		check(InOutBatchHeader.NumElements <= UE::MorphTargetVertexCodec::BatchSize);
-		TArrayView<UE::MorphTargetVertexCodec::FQuantizedDelta> BatchQuantizedDeltas(
-				QuantizedDeltasReadBuffer, InOutBatchHeader.NumElements);
+		int32 NumUsedVertices = UsedVertices.CountSetBits(BatchVertexRangeBegin, BatchVertexRangeEnd); 
 
-		ReadQuantizedDeltas(BatchQuantizedDeltas, InOutBatchHeader, BatchDataView); 
-		
-		const int32 NumDeltas = InOutBatchHeader.NumElements;
+		if (NumUsedVertices == 0)
+		{
+			// All vertices removed.
+			InOutBatchHeader.IndexMin = 0;
+			InOutBatchHeader.NumElements = 0;
+			InOutBatchHeader.DataOffset = 0;
+		}
+		else if (NumUsedVertices == BatchVertexRangeEnd - BatchVertexRangeBegin)
+		{
+			// No vertices removed, remap Batch IndexMin.
+			InOutBatchHeader.IndexMin = VertexRemap[InOutBatchHeader.IndexMin];
+			check(InOutBatchHeader.IndexMin != INDEX_NONE);
+		}
+		else
+		{
+			TArrayView<uint32> BatchDataView = TArrayView<uint32>(
+					DataView.GetData() + (InOutBatchHeader.DataOffset/sizeof(uint32)), 
+					CalculateBatchDwords(InOutBatchHeader)); 
 
-		check(InOutBatchHeader.NumElements <= UE::MorphTargetVertexCodec::BatchSize);
+			check(InOutBatchHeader.NumElements <= UE::MorphTargetVertexCodec::BatchSize);
+			TArrayView<UE::MorphTargetVertexCodec::FQuantizedDelta> BatchQuantizedDeltas(
+					QuantizedDeltasReadBuffer, InOutBatchHeader.NumElements);
 
-		InOutBatchHeader.IndexMin = 0;
-		
-		int32 DeltaIndex = 0;
-		for ( ; DeltaIndex < NumDeltas; ++DeltaIndex)
-		{
-			check((int32)BatchQuantizedDeltas[DeltaIndex].Index < VertexRemap.Num());
+			ReadQuantizedDeltas(BatchQuantizedDeltas, InOutBatchHeader, BatchDataView); 
 			
-			int32 RemappedDeltaVertexIdx = VertexRemap[BatchQuantizedDeltas[DeltaIndex].Index];
-			if (RemappedDeltaVertexIdx >= 0)
+			const int32 NumDeltas = InOutBatchHeader.NumElements;
+			check(NumDeltas <= UE::MorphTargetVertexCodec::BatchSize);
+
+			InOutBatchHeader.IndexMin = 0;
+		 
+			int32 DeltaIndex = 0;
+			for ( ; DeltaIndex < NumDeltas; ++DeltaIndex)
 			{
-				check(RemappedDeltaVertexIdx < NumResultVertices);
-				InOutBatchHeader.IndexMin = (uint32)RemappedDeltaVertexIdx;
-				break;
-			}
-		}
+				check((int32)BatchQuantizedDeltas[DeltaIndex].Index < UsedVertices.Num());
+				
+				if (UsedVertices[BatchQuantizedDeltas[DeltaIndex].Index])
+				{
+					int32 RemappedDeltaVertexIdx = VertexRemap[BatchQuantizedDeltas[DeltaIndex].Index];
+					check(RemappedDeltaVertexIdx < NumResultVertices);
+					InOutBatchHeader.IndexMin = (uint32)RemappedDeltaVertexIdx;
 
-		UE::MorphTargetVertexCodec::FQuantizedDelta QuantizedDeltasKeepBuffer[UE::MorphTargetVertexCodec::BatchSize];
+					break;
+				}
+			}
 
-		uint32 DestDeltaIndex = 0;
-		for ( ; DeltaIndex < NumDeltas; ++DeltaIndex)
-		{
-			check((int32)BatchQuantizedDeltas[DeltaIndex].Index < VertexRemap.Num());
+			UE::MorphTargetVertexCodec::FQuantizedDelta QuantizedDeltasKeepBuffer[UE::MorphTargetVertexCodec::BatchSize];
 
-			int32 RemappedDeltaVertexIdx = VertexRemap[BatchQuantizedDeltas[DeltaIndex].Index];
-			if (RemappedDeltaVertexIdx >= 0)
+			uint32 DestDeltaIndex = 0;
+			for ( ; DeltaIndex < NumDeltas; ++DeltaIndex)
 			{
-				check(RemappedDeltaVertexIdx < NumResultVertices);
-
-				QuantizedDeltasKeepBuffer[DestDeltaIndex++] = UE::MorphTargetVertexCodec::FQuantizedDelta
+				if (UsedVertices[BatchQuantizedDeltas[DeltaIndex].Index])
 				{
-					.Position = BatchQuantizedDeltas[DeltaIndex].Position,
-					.TangentZ = BatchQuantizedDeltas[DeltaIndex].TangentZ,
-					.Index    = (uint32)RemappedDeltaVertexIdx,
-				};	
+					check((int32)BatchQuantizedDeltas[DeltaIndex].Index < VertexRemap.Num());
+					int32 RemappedDeltaVertexIdx = VertexRemap[BatchQuantizedDeltas[DeltaIndex].Index];
+
+					check(RemappedDeltaVertexIdx < NumResultVertices);
+
+					QuantizedDeltasKeepBuffer[DestDeltaIndex++] = UE::MorphTargetVertexCodec::FQuantizedDelta
+					{
+						.Position = BatchQuantizedDeltas[DeltaIndex].Position,
+						.TangentZ = BatchQuantizedDeltas[DeltaIndex].TangentZ,
+						.Index    = (uint32)RemappedDeltaVertexIdx,
+					};	
+				}
 			}
-		}
-		
-		InOutBatchHeader.NumElements = DestDeltaIndex;
-		check(DestDeltaIndex <= UE::MorphTargetVertexCodec::BatchSize);
-		WriteQuantizedDeltas(TArrayView(QuantizedDeltasKeepBuffer, DestDeltaIndex), InOutBatchHeader, BatchDataView);	
 
+			InOutBatchHeader.NumElements = DestDeltaIndex;
+
+			if (InOutBatchHeader.NumElements == 0)
+			{
+				InOutBatchHeader.IndexMin = 0;
+				InOutBatchHeader.DataOffset = 0;
+			}
+			else
+			{
+				check(DestDeltaIndex <= UE::MorphTargetVertexCodec::BatchSize);
+				WriteQuantizedDeltas(TArrayView(QuantizedDeltasKeepBuffer, DestDeltaIndex), InOutBatchHeader, BatchDataView);
+			}
+		}
 #if 0 // Check data integrity.
 		{
 			TArray<UE::MorphTargetVertexCodec::FQuantizedDelta, TInlineAllocator<UE::MorphTargetVertexCodec::BatchSize>> TestQuantizedDeltasStorage;
@@ -101,9 +165,9 @@ namespace UE::Mutable::Private
 	}
 
 
-	void MeshRemoveMeshMorphs(FMesh& Result, TConstArrayView<int32> VertexIndexRemap)
+	void MeshRemoveMeshMorphs(FMesh& Result, const TBitArray<>& UsedVertices, TConstArrayView<int32> VertexIndexRemap)
 	{
-		MUTABLE_CPUPROFILER_SCOPE(CompressedMorphs_MeshRemoveMesh)
+		MUTABLE_CPUPROFILER_SCOPE(CompressedMorphs_RemoveMesh)
 
 		using namespace UE::MorphTargetVertexCodec;
 
@@ -113,82 +177,128 @@ namespace UE::Mutable::Private
 		}
 
 		int32 NumMorphs = Result.Morph.Names.Num();
-		for (int32 MorphIndex = 0; MorphIndex < NumMorphs; ++MorphIndex)
+		TArrayView<uint32> DataView = Result.MorphDataBuffer; 
+		uint32* MorphDataPtr = DataView.GetData();
+		
+		// Init to false so any empty batch coming from other operations are also removed.
+		TBitArray<> BatchesToKeep;
+		BatchesToKeep.Init(false, Result.Morph.NumTotalBatches);
 		{
-			const int32 NumBatches = Result.Morph.BatchesPerMorph[MorphIndex];
-			
-			if (NumBatches == 0)
+			MUTABLE_CPUPROFILER_SCOPE(CompressedMorphs_RemoveMesh_RemoveVertices)
+			for (int32 MorphIndex = 0; MorphIndex < NumMorphs; ++MorphIndex)
 			{
-				continue;
+				const int32 NumBatches = Result.Morph.BatchesPerMorph[MorphIndex];
+				
+				if (NumBatches == 0)
+				{
+					continue;
+				}
+
+				int32 BatchStartOffset = Result.Morph.BatchStartOffsetPerMorph[MorphIndex];
+	
+				FDeltaBatchHeader BatchHeaderStorage[2];
+
+				FDeltaBatchHeader* CurrBatchHeader = &BatchHeaderStorage[0];
+				FDeltaBatchHeader* NextBatchHeader = &BatchHeaderStorage[1];
+				
+				int32 CurrBatchIndex = 0;
+				//Find first non empty batch.
+				for (; CurrBatchIndex < NumBatches; ++CurrBatchIndex)
+				{
+					ReadHeader(*CurrBatchHeader, TArrayView<uint32>(MorphDataPtr + ((BatchStartOffset + CurrBatchIndex) * NumBatchHeaderDwords), NumBatchHeaderDwords));
+ 
+					if (CurrBatchHeader->NumElements)
+					{
+						break;
+					}
+				}
+
+				const int32 NumResultVertices = Result.GetVertexCount(); 
+				const int32 NumSourceVertices = VertexIndexRemap.Num();
+				for (int32 NextBatchIndex = CurrBatchIndex + 1; NextBatchIndex <= NumBatches; ++NextBatchIndex)
+				{
+					// Find next non empty batch.
+					for (; NextBatchIndex < NumBatches; ++NextBatchIndex)
+					{
+						ReadHeader(*NextBatchHeader, TArrayView<uint32>(MorphDataPtr + (BatchStartOffset + NextBatchIndex) * NumBatchHeaderDwords, NumBatchHeaderDwords));
+
+						if (NextBatchHeader->NumElements)
+						{
+							break;
+						}
+					}
+			
+					int32 BatchRangeBegin = CurrBatchHeader->IndexMin;
+					int32 BatchRangeEnd   = NextBatchIndex < NumBatches ? NextBatchHeader->IndexMin : NumSourceVertices; 	
+
+					RemoveMorphBatchVertices(*CurrBatchHeader, DataView, BatchRangeBegin, BatchRangeEnd, UsedVertices, VertexIndexRemap, NumResultVertices);
+
+					WriteHeader(*CurrBatchHeader, TArrayView<uint32>(MorphDataPtr + ((BatchStartOffset + CurrBatchIndex) * NumBatchHeaderDwords), NumBatchHeaderDwords));
+					BatchesToKeep[BatchStartOffset + CurrBatchIndex] = CurrBatchHeader->NumElements != 0;
+
+					::Swap(CurrBatchHeader, NextBatchHeader);
+					CurrBatchIndex = NextBatchIndex;
+				}
 			}
+		}
 
-			int32 BatchStartOffset = Result.Morph.BatchStartOffsetPerMorph[MorphIndex];
+		int32 CumulativeDataOffsetInDwords = 0;
+		int32 CumulativeBatchStartOffset = 0;
+		{
+			MUTABLE_CPUPROFILER_SCOPE(CompressedMorphs_RemoveMesh_MoveHeadersAndMetadata);
 
-			TArrayView<uint32> DataView = Result.MorphDataBuffer; 
-			uint32* MorphDataPtr = DataView.GetData();
+			CumulativeDataOffsetInDwords = RemoveBufferUnusedElems(
+					(uint8*)MorphDataPtr, NumBatchHeaderDwords*sizeof(uint32), BatchesToKeep.Num(), BatchesToKeep) * NumBatchHeaderDwords; 
 
-			for (int32 BatchIndex = 0; BatchIndex < NumBatches; ++BatchIndex)
+			// Fix Morph metadata.
+			for (int32 MorphIndex = 0; MorphIndex < NumMorphs; ++MorphIndex)
 			{
-				FDeltaBatchHeader BatchHeader;
-				ReadHeader(BatchHeader, TArrayView<uint32>(MorphDataPtr + ((BatchStartOffset + BatchIndex) * NumBatchHeaderDwords), NumBatchHeaderDwords));
-				RemoveMorphBatchVertices(BatchHeader, DataView, VertexIndexRemap, Result.GetVertexCount());	
+				int32 MorphBatchesBegin = Result.Morph.BatchStartOffsetPerMorph[MorphIndex];
+				int32 MorphBatchesEnd   = MorphBatchesBegin + Result.Morph.BatchesPerMorph[MorphIndex];
+
+				int32 NumMorphBatches = BatchesToKeep.CountSetBits(MorphBatchesBegin, MorphBatchesEnd);
+				Result.Morph.BatchesPerMorph[MorphIndex] = NumMorphBatches; 
 
-				WriteHeader(BatchHeader, TArrayView<uint32>(MorphDataPtr + ((BatchStartOffset + BatchIndex) * NumBatchHeaderDwords), NumBatchHeaderDwords));
+				Result.Morph.BatchStartOffsetPerMorph[MorphIndex] = CumulativeBatchStartOffset;
+				CumulativeBatchStartOffset += NumMorphBatches;
 			}
 
-			// TODO: 
-			// The following implementation is better as it can skip some work by checking the range 
-			// of indices for a batch only looking at 2 consecutive headers. This should allow to remap all
-			// morph vertices indices just changing the batch MinIndex for batches that don't remove vertices 
-			// but this approach does not seam to be working. The code above remaps all batches changing the 
-			// deltas, but it world be worth investigate why the more optimal approach is not working and 
-			// reimplement it.
-			//
-			// Additionally, this code does not remove data from the batches, it only modifies the number of 
-			// valid deltas in the batches. As a memory optimization, the unused data should be removed.
-			
-			//FDeltaBatchHeader BatchHeaderStorage[2];
-
-			//FDeltaBatchHeader* CurrBatchHeader = &BatchHeaderStorage[0];
-			//FDeltaBatchHeader* NextBatchHeader = &BatchHeaderStorage[1];
-			//
-			//int32 CurrBatchIndex = 0;
-			// Find first non empty batch.
-			//for (; CurrBatchIndex < NumBatches; ++CurrBatchIndex)
-			//{
-			//	ReadHeader(*CurrBatchHeader, TArrayView<uint32>(MorphDataPtr + ((BatchStartOffset + CurrBatchIndex) * NumBatchHeaderDwords), NumBatchHeaderDwords));
-
-			//	if (CurrBatchHeader->NumElements)
-			//	{
-			//		break;
-			//	}
-			//}
-
-			//const int32 NumSourceVertices = VertexIndexRemap.Num();
-			//for (int32 NextBatchIndex = CurrBatchIndex + 1; NextBatchIndex < NumBatches; ++NextBatchIndex)
-			//{
-			//	// Find next non empty batch.
-			//	for (; NextBatchIndex < NumBatches; ++NextBatchIndex)
-			//	{
-			//		ReadHeader(*NextBatchHeader, TArrayView<uint32>(MorphDataPtr + ((BatchStartOffset + NextBatchIndex) * NumBatchHeaderDwords), NumBatchHeaderDwords));
-
-			//		if (NextBatchHeader->NumElements)
-			//		{
-			//			break;
-			//		}
-			//	}
-		
-			//	int32 BatchRangeBegin = CurrBatchHeader->IndexMin;
-			//	int32 BatchRangeEnd   = NextBatchIndex < NumBatches ? NextBatchHeader->IndexMin : NumSourceVertices; 	
+			Result.Morph.NumTotalBatches = CumulativeBatchStartOffset;
+		}
+		{
+			MUTABLE_CPUPROFILER_SCOPE(CompressedMorphs_RemoveMesh_MoveData);
 
-			//	RemoveMorphBatchVertices(*CurrBatchHeader, DataView, BatchRangeBegin, BatchRangeEnd, VertexIndexRemap);	
+			check(CumulativeDataOffsetInDwords == CumulativeBatchStartOffset*NumBatchHeaderDwords);
+			for (int32 BatchIndex = 0; BatchIndex < CumulativeBatchStartOffset; ++BatchIndex)
+			{
+				const int32 BatchDataOffsetInDwords = BatchIndex*NumBatchHeaderDwords;
 
-			//	WriteHeader(*CurrBatchHeader, TArrayView<uint32>(MorphDataPtr + ((BatchStartOffset + CurrBatchIndex) * NumBatchHeaderDwords), NumBatchHeaderDwords));
+				FDeltaBatchHeader Header;
+				ReadHeader(Header, TArrayView<uint32>(MorphDataPtr + BatchDataOffsetInDwords, NumBatchHeaderDwords));
+				int32 NumBatchDataDwords = CalculateBatchDwords(Header);
+			
+				int32 OriginalDataOffsetInDwords = Header.DataOffset / sizeof(uint32);
+				if (CumulativeDataOffsetInDwords != OriginalDataOffsetInDwords)
+				{
+					check(CumulativeDataOffsetInDwords + NumBatchDataDwords <= Result.MorphDataBuffer.Num());
+					FMemory::Memmove(
+							MorphDataPtr + CumulativeDataOffsetInDwords,
+							MorphDataPtr + OriginalDataOffsetInDwords,
+							NumBatchDataDwords * sizeof(uint32));
+				
+					// First Batch dword corresponds to the data offset. See MorphTargetVertexCodec::WriteHeader()
+					MorphDataPtr[BatchIndex*NumBatchHeaderDwords] = CumulativeDataOffsetInDwords * sizeof(uint32);
+				}
+
+				CumulativeDataOffsetInDwords += NumBatchDataDwords;
+			}	
+		}
 
-			//	::Swap(CurrBatchHeader, NextBatchHeader);
-			//	CurrBatchIndex = NextBatchIndex;
-			//}
-		}	
+		// Trim unused bytes. Don't free memory here, delay it until a copy is done.
+		{
+			MUTABLE_CPUPROFILER_SCOPE(CompressedMorphs_RemoveMesh_Trim);
+			Result.MorphDataBuffer.SetNum(CumulativeDataOffsetInDwords, EAllowShrinking::No);
+		}
 	}
 
     void MeshRemoveRecreateSurface(FMesh* Result, const TBitArray<>& UsedVertices, const TBitArray<>& UsedFaces)
@@ -349,42 +459,7 @@ namespace UE::Mutable::Private
             }
         }
         return -1;
-    }
-
-	
-	void CopyMoveBuffer(uint8* BufferData, uint32 ElemSize, int32 NumVertices, const TBitArray<>& UsedVertices)
-	{
-    	if (ElemSize == 0)
-    	{
-    		return;
-    	}
-    	
-		int32 LastFreeVertexIndex = 0;
-
-		for (int32 VertexIndex = UsedVertices.FindFrom(true, 0); VertexIndex >= 0;)
-		{
-			const int32 UsedSpanBegin = VertexIndex;
-			VertexIndex = UsedVertices.FindFrom(false, VertexIndex);
-			VertexIndex = VertexIndex >= 0 ? VertexIndex : NumVertices;
-			const int32 UsedSpanEnd = VertexIndex;
-
-			VertexIndex = UsedVertices.FindFrom(true, VertexIndex);
-
-			// Copy vertex buffer span.	
-			const int32 UsedSpanSize = UsedSpanEnd - UsedSpanBegin;
-            
-			if (LastFreeVertexIndex != UsedSpanBegin)
-			{
-				FMemory::Memmove(
-						BufferData + LastFreeVertexIndex*ElemSize,
-						BufferData + UsedSpanBegin*ElemSize,
-						UsedSpanSize*ElemSize);
-			}
-
-			LastFreeVertexIndex += UsedSpanSize;
-		}
-	}
-	
+    }	
 	
     void MeshRemoveVerticesWithCullSet(FMesh* Result, const TBitArray<>& VerticesToCull, bool bRemoveIfAllVerticesCulled)
     { 
@@ -571,7 +646,7 @@ namespace UE::Mutable::Private
             	uint8* BufferData = VertexBufferSet.GetBufferData(BufferIndex); 
             	const uint32 ElemSize = VertexBufferSet.GetElementSize(BufferIndex);
 
-            	CopyMoveBuffer(BufferData, ElemSize, NumVertices, UsedVertices);
+            	RemoveBufferUnusedElems(BufferData, ElemSize, NumVertices, UsedVertices);
             }
         }
 		Result->GetVertexBuffers().SetElementCount(NumVerticesRemaining);
@@ -585,7 +660,7 @@ namespace UE::Mutable::Private
     		const int32 NumInfluences = NumVerticesRemaining ? NumClothVertices / NumVerticesRemaining : 0;
     		const uint32 ElemSize = sizeof(FMeshToMeshVertData) * NumInfluences;
 
-    		CopyMoveBuffer(BufferData, ElemSize, NumVertices, UsedVertices);
+    		RemoveBufferUnusedElems(BufferData, ElemSize, NumVertices, UsedVertices);
     		Result->ClothSections[0].Data.SetNum(NumVerticesRemaining * NumInfluences);
     	}
 
@@ -673,7 +748,7 @@ namespace UE::Mutable::Private
 
 		MeshRemoveRecreateSurface(Result, UsedVertices, UsedFaces);
 
-		MeshRemoveMeshMorphs(*Result, UsedVerticesMap);
+		MeshRemoveMeshMorphs(*Result, UsedVertices, UsedVerticesMap);
 
     }
 
