[Mutable] Added additional ASTOp optimizations.

#Synced-CL 51880525

diff --git a/Engine/Plugins/Mutable/Source/CustomizableObject/Internal/MuCO/CustomizableObjectPrivate.h b/Engine/Plugins/Mutable/Source/CustomizableObject/Internal/MuCO/CustomizableObjectPrivate.h
--- a/Engine/Plugins/Mutable/Source/CustomizableObject/Internal/MuCO/CustomizableObjectPrivate.h
+++ b/Engine/Plugins/Mutable/Source/CustomizableObject/Internal/MuCO/CustomizableObjectPrivate.h
@@ -1457,6 +1457,8 @@ enum class ECustomizableObjectVersions : int32
 	
 	IS_SWITCH,
 	
+	SkeletalMeshConditionalSwitch,
+	
 	// -----<new versions can be added above this line>--------
 	LastCustomizableObjectVersion
 };
diff --git a/Engine/Plugins/Mutable/Source/CustomizableObjectEditor/Private/MuCOE/SMutableCodeViewer.cpp b/Engine/Plugins/Mutable/Source/CustomizableObjectEditor/Private/MuCOE/SMutableCodeViewer.cpp
--- a/Engine/Plugins/Mutable/Source/CustomizableObjectEditor/Private/MuCOE/SMutableCodeViewer.cpp
+++ b/Engine/Plugins/Mutable/Source/CustomizableObjectEditor/Private/MuCOE/SMutableCodeViewer.cpp
@@ -1420,6 +1420,8 @@ void SMutableCodeViewer::GenerateElement(const int32& InStateIndex, UE::Mutable:
 		case UE::Mutable::Private::EOpType::IN_CONDITIONAL:
 		case UE::Mutable::Private::EOpType::ED_CONDITIONAL:
 		case UE::Mutable::Private::EOpType::MI_CONDITIONAL:
+		case UE::Mutable::Private::EOpType::IS_CONDITIONAL:
+		case UE::Mutable::Private::EOpType::SK_CONDITIONAL:
 			{
 				uint32 ChildIndex = 0;
 			
@@ -1438,6 +1440,9 @@ void SMutableCodeViewer::GenerateElement(const int32& InStateIndex, UE::Mutable:
 		case UE::Mutable::Private::EOpType::NU_SWITCH:
 		case UE::Mutable::Private::EOpType::IN_SWITCH:
 		case UE::Mutable::Private::EOpType::ED_SWITCH:
+		case UE::Mutable::Private::EOpType::SK_SWITCH:
+		case UE::Mutable::Private::EOpType::IS_SWITCH:
+		case UE::Mutable::Private::EOpType::MI_SWITCH:
 			{
 				uint32 ChildIndex = 0;
 
diff --git a/Engine/Plugins/Mutable/Source/MutableRuntime/Internal/MuR/CodeVisitor.h b/Engine/Plugins/Mutable/Source/MutableRuntime/Internal/MuR/CodeVisitor.h
--- a/Engine/Plugins/Mutable/Source/MutableRuntime/Internal/MuR/CodeVisitor.h
+++ b/Engine/Plugins/Mutable/Source/MutableRuntime/Internal/MuR/CodeVisitor.h
@@ -259,6 +259,8 @@ namespace UE::Mutable::Private
             case EOpType::IN_CONDITIONAL:
 			case EOpType::ED_CONDITIONAL:
 			case EOpType::MI_CONDITIONAL:
+			case EOpType::SK_CONDITIONAL:
+			case EOpType::IS_CONDITIONAL:
             {
 				FOperation::ConditionalArgs Args = Program.GetOpArgs<FOperation::ConditionalArgs>(Address);
 
@@ -293,6 +295,7 @@ namespace UE::Mutable::Private
             case EOpType::LA_SWITCH:
             case EOpType::IN_SWITCH:
 			case EOpType::ED_SWITCH:
+			case EOpType::SK_SWITCH:
             {
                 bRecurse = false;
 
diff --git a/Engine/Plugins/Mutable/Source/MutableRuntime/Internal/MuR/Operations.h b/Engine/Plugins/Mutable/Source/MutableRuntime/Internal/MuR/Operations.h
--- a/Engine/Plugins/Mutable/Source/MutableRuntime/Internal/MuR/Operations.h
+++ b/Engine/Plugins/Mutable/Source/MutableRuntime/Internal/MuR/Operations.h
@@ -397,7 +397,10 @@ namespace UE::Mutable::Private
         ME_RANGE_END   = SK_PARAMETER,
         SK_RANGE_BEGIN = ME_RANGE_END,
 
+    	SK_CONDITIONAL,
 
+    	SK_SWITCH,
+    	
         //-----------------------------------------------------------------------------------------
         // Instance operations
         //-----------------------------------------------------------------------------------------
@@ -546,6 +549,8 @@ namespace UE::Mutable::Private
         /** External operation */
 		IS_EXTERNAL,
 
+    	IS_CONDITIONAL,
+    	
     	IS_SWITCH,
     
         IS_RANGE_END,
@@ -1548,26 +1553,51 @@ namespace UE::Mutable::Private
     //! Utility function to apply a function to all operation references to other operations.
 	MUTABLERUNTIME_API extern void ForEachReference( const struct FProgram& Program, FOperation::ADDRESS At, const TFunctionRef<void(FOperation::ADDRESS)> );
 
-	//!
 	inline EOpType GetSwitchForType( EDataType d )
-    {
-        switch (d)
-        {
+	{
+		switch (d)
+		{
 		case EDataType::Instance: return EOpType::IN_SWITCH;
-        case EDataType::Mesh: return EOpType::ME_SWITCH;
+		case EDataType::Mesh: return EOpType::ME_SWITCH;
 		case EDataType::Image: return EOpType::IM_SWITCH;
 		case EDataType::Layout: return EOpType::LA_SWITCH;
-        case EDataType::Color: return EOpType::CO_SWITCH;
-        case EDataType::Scalar: return EOpType::SC_SWITCH;
-        case EDataType::Int: return EOpType::NU_SWITCH;
+		case EDataType::Color: return EOpType::CO_SWITCH;
+		case EDataType::Scalar: return EOpType::SC_SWITCH;
+		case EDataType::Int: return EOpType::NU_SWITCH;
 		case EDataType::ExtensionData: return EOpType::ED_SWITCH;
 		case EDataType::InstancedStruct: return EOpType::IS_SWITCH;
+		case EDataType::SkeletalMesh: return EOpType::SK_SWITCH;
+		case EDataType::Material: return EOpType::MI_SWITCH;
 
-        default:
+		default:
 			check(false);
 			break;
-        }
-        return EOpType::NONE;
-    }
+		}
+		return EOpType::NONE;
+	}
+	
+	inline EOpType GetConditionalForType(EDataType Type)
+	{
+		switch (Type)
+		{
+		case EDataType::Instance: return EOpType::IN_CONDITIONAL;
+		case EDataType::Mesh: return EOpType::ME_CONDITIONAL;
+		case EDataType::Image: return EOpType::IM_CONDITIONAL;
+		case EDataType::Layout: return EOpType::LA_CONDITIONAL;
+		case EDataType::Color: return EOpType::CO_CONDITIONAL;
+		case EDataType::Scalar: return EOpType::SC_CONDITIONAL;
+		case EDataType::Int: return EOpType::NU_CONDITIONAL;
+		case EDataType::ExtensionData: return EOpType::ED_CONDITIONAL;
+		case EDataType::InstancedStruct: return EOpType::IS_CONDITIONAL;
+		case EDataType::SkeletalMesh: return EOpType::SK_CONDITIONAL;
+		case EDataType::Material: return EOpType::MI_CONDITIONAL;
+
+		default:
+			check(false);
+			break;
+		}
+		
+		return EOpType::NONE;
+	}
 }
 
diff --git a/Engine/Plugins/Mutable/Source/MutableRuntime/Private/MuR/CodeRunner.cpp b/Engine/Plugins/Mutable/Source/MutableRuntime/Private/MuR/CodeRunner.cpp
--- a/Engine/Plugins/Mutable/Source/MutableRuntime/Private/MuR/CodeRunner.cpp
+++ b/Engine/Plugins/Mutable/Source/MutableRuntime/Private/MuR/CodeRunner.cpp
@@ -340,6 +340,7 @@ namespace UE::Mutable::Private
             case EDataType::Instance:		StoreInstance( cat, LoadInstance(rat) ); break;
 			case EDataType::ExtensionData:	StoreExtensionData( cat, LoadExtensionData(rat) ); break;
 			case EDataType::Material:		StoreMaterial( cat, LoadMaterial(rat) ); break;
+			case EDataType::SkeletalMesh:	StoreSkeletalMesh( cat, LoadSkeletalMesh(rat) ); break;
             default:
                 // Not implemented
                 check( false );
@@ -398,6 +399,7 @@ namespace UE::Mutable::Private
 				case EDataType::ExtensionData:		StoreExtensionData(Item, MakeManaged<FExtensionData>()); break;
 				case EDataType::Material:			StoreMaterial(Item, MakeManaged<FMaterial>()); break;
 				case EDataType::InstancedStruct:	StoreInstancedStruct(Item, MakeManaged<FInstancedStruct>()); break;
+				case EDataType::SkeletalMesh:		StoreSkeletalMesh(Item, MakeManaged<FSkeletalMesh>()); break;
 				default:
 					// Not implemented
 					check(false);
@@ -486,6 +488,7 @@ namespace UE::Mutable::Private
 			case EDataType::ExtensionData:		StoreExtensionData(ItemAddress, LoadExtensionData(ResultAddress)); break;
 			case EDataType::Material:			StoreMaterial(ItemAddress, LoadMaterial(ResultAddress)); break;
 			case EDataType::InstancedStruct:	StoreInstancedStruct(ItemAddress, LoadInstancedStruct(ResultAddress)); break;
+			case EDataType::SkeletalMesh:		StoreSkeletalMesh(ItemAddress, LoadSkeletalMesh(ResultAddress)); break;
             default:
                 // Not implemented
                 check(false);
@@ -7098,6 +7101,8 @@ namespace UE::Mutable::Private
         case EOpType::IN_CONDITIONAL:
 		case EOpType::ED_CONDITIONAL:
 		case EOpType::MI_CONDITIONAL:
+		case EOpType::SK_CONDITIONAL:
+		case EOpType::IS_CONDITIONAL:
             RunCode_Conditional(Item);
             break;
 
@@ -7116,6 +7121,7 @@ namespace UE::Mutable::Private
 		case EOpType::ED_SWITCH:
 		case EOpType::MI_SWITCH:
 		case EOpType::IS_SWITCH:
+		case EOpType::SK_SWITCH:
             RunCode_Switch(Item);
             break;
 
diff --git a/Engine/Plugins/Mutable/Source/MutableRuntime/Private/MuR/CodeRunnerBegin.cpp b/Engine/Plugins/Mutable/Source/MutableRuntime/Private/MuR/CodeRunnerBegin.cpp
--- a/Engine/Plugins/Mutable/Source/MutableRuntime/Private/MuR/CodeRunnerBegin.cpp
+++ b/Engine/Plugins/Mutable/Source/MutableRuntime/Private/MuR/CodeRunnerBegin.cpp
@@ -325,6 +325,8 @@ namespace UE::Mutable::Private
 			case EOpType::SC_CONDITIONAL:
 			case EOpType::MI_CONDITIONAL:
 			case EOpType::NU_CONDITIONAL:
+			case EOpType::SK_CONDITIONAL:
+			case EOpType::IS_CONDITIONAL:
 				{
 					FOperation::ConditionalArgs Args = Program.GetOpArgs<FOperation::ConditionalArgs>(Item.At);
 
@@ -377,6 +379,7 @@ namespace UE::Mutable::Private
 								case EDataType::String:
 								case EDataType::Instance:
 								case EDataType::Material:
+								case EDataType::SkeletalMesh:
 									StoreNone(Item);
 									break;
 
@@ -412,6 +415,7 @@ namespace UE::Mutable::Private
 			case EOpType::SC_SWITCH:
 			case EOpType::MI_SWITCH:
 			case EOpType::IS_SWITCH:
+			case EOpType::SK_SWITCH:
 				{
 					const uint8* Data = Program.GetOpArgsPointer(Item.At);
 
@@ -464,6 +468,7 @@ namespace UE::Mutable::Private
 									case EDataType::Instance:
 									case EDataType::String:
 									case EDataType::Material:
+									case EDataType::SkeletalMesh:
 										StoreNone(Item);
 										break;
 								
diff --git a/Engine/Plugins/Mutable/Source/MutableRuntime/Private/MuR/Operations.cpp b/Engine/Plugins/Mutable/Source/MutableRuntime/Private/MuR/Operations.cpp
--- a/Engine/Plugins/Mutable/Source/MutableRuntime/Private/MuR/Operations.cpp
+++ b/Engine/Plugins/Mutable/Source/MutableRuntime/Private/MuR/Operations.cpp
@@ -80,6 +80,7 @@ namespace UE::Mutable::Private
 		case EOpType::ED_SWITCH:
 		case EOpType::MI_SWITCH:
 		case EOpType::IS_SWITCH:
+		case EOpType::SK_SWITCH:
         {
 			const uint8* Data = program.GetOpArgsPointer(at);
 			
diff --git a/Engine/Plugins/Mutable/Source/MutableTools/Private/MuT/ASTOpConditional.cpp b/Engine/Plugins/Mutable/Source/MutableTools/Private/MuT/ASTOpConditional.cpp
--- a/Engine/Plugins/Mutable/Source/MutableTools/Private/MuT/ASTOpConditional.cpp
+++ b/Engine/Plugins/Mutable/Source/MutableTools/Private/MuT/ASTOpConditional.cpp
@@ -87,6 +87,8 @@ namespace UE::Mutable::Private
 		case EOpType::IN_CONDITIONAL:
 		case EOpType::ED_CONDITIONAL:
 		case EOpType::MI_CONDITIONAL:
+		case EOpType::SK_CONDITIONAL:
+		case EOpType::IS_CONDITIONAL:
 			break;
 		default:
 			// Unexpected type
diff --git a/Engine/Plugins/Mutable/Source/MutableTools/Private/MuT/ASTOpSwitch.cpp b/Engine/Plugins/Mutable/Source/MutableTools/Private/MuT/ASTOpSwitch.cpp
--- a/Engine/Plugins/Mutable/Source/MutableTools/Private/MuT/ASTOpSwitch.cpp
+++ b/Engine/Plugins/Mutable/Source/MutableTools/Private/MuT/ASTOpSwitch.cpp
@@ -2,6 +2,8 @@
 
 #include "MuT/ASTOpSwitch.h"
 
+#include "ASTOpBoolEqualIntConst.h"
+#include "ASTOpConditional.h"
 #include "MuT/ASTOpParameter.h"
 #include "MuT/ASTOpConstantInt.h"
 #include "MuT/ASTOpMeshAddMetadata.h"
@@ -74,6 +76,9 @@ namespace UE::Mutable::Private
 		case EOpType::LA_SWITCH:
 		case EOpType::IN_SWITCH:
 		case EOpType::ED_SWITCH:
+		case EOpType::SK_SWITCH:
+		case EOpType::IS_SWITCH:
+		case EOpType::MI_SWITCH:
 			break;
 		default:
 			// Unexpected Type
@@ -668,7 +673,45 @@ namespace UE::Mutable::Private
 				// We can remove this switch, all paths leading to it have the same Condition for this switches Variable.
 				return FindBranch(MatchingSwitchOption);
 			}
+		}
+
+		// Remove conditionals with the same condition as the switch.
+		{
+			Ptr<ASTOpSwitch> NewSwitch = nullptr;
 
+			for (int32 CaseIndex = 0; CaseIndex < Cases.Num(); ++CaseIndex)
+			{
+				const FCase& Case = Cases[CaseIndex];
+				
+				if (!Case.Branch || !Case.Branch->IsConditional())
+				{
+					continue;
+				}
+				
+				ASTOpConditional* Conditional = static_cast<ASTOpConditional*>(Case.Branch.child().get());
+				if (!Conditional->condition || Conditional->condition->GetOpType() != EOpType::BO_EQUAL_INT_CONST)
+				{
+					continue;
+				}
+				
+				ASTOpBoolEqualIntConst* Condition = static_cast<ASTOpBoolEqualIntConst*>(Conditional->condition.child().get());
+				if (Condition->Value != Variable)
+				{
+					continue;
+				}
+				
+				if (!NewSwitch)
+				{
+					NewSwitch = UE::Mutable::Private::Clone<ASTOpSwitch>(this);
+				}
+				
+				NewSwitch->Cases[CaseIndex].Branch = Conditional->yes.child();
+			}
+			
+			if (NewSwitch)
+			{
+				return NewSwitch;
+			}
 		}
 
 		return nullptr;
diff --git a/Engine/Plugins/Mutable/Source/MutableTools/Private/MuT/CodeOptimiser_Logic.cpp b/Engine/Plugins/Mutable/Source/MutableTools/Private/MuT/CodeOptimiser_Logic.cpp
--- a/Engine/Plugins/Mutable/Source/MutableTools/Private/MuT/CodeOptimiser_Logic.cpp
+++ b/Engine/Plugins/Mutable/Source/MutableTools/Private/MuT/CodeOptimiser_Logic.cpp
@@ -25,131 +25,63 @@ namespace UE::Mutable::Private
 
         // The separate steps should not be combined into one traversal
 
-        // Unwrap some typical code daisy-chains
+		// Unwrap some typical code daisy-chains
         //-----------------------------------------------------------------------------------------
         {
-            MUTABLE_CPUPROFILER_SCOPE(Unwrap);
-            ASTOp::Traverse_TopDown_Unique_Imprecise( roots, [&](Ptr<ASTOp>& n)
+            MUTABLE_CPUPROFILER_SCOPE(ContionalExclusive);
+            ASTOp::Traverse_TopDown_Unique_Imprecise(roots, [&](Ptr<ASTOp>& Op)
             {
-                bool bRecurse = true;
-
-                if (n->GetOpType()==EOpType::IN_CONDITIONAL)
-                {
-					ASTOpConditional* topConditional = static_cast<ASTOpConditional*>(n.get());
-					Ptr<ASTOp> yes = topConditional->yes.child();
-                    if (yes && yes->GetOpType()==EOpType::IN_ADDSURFACE)
-                    {
-						ASTOpInstanceAdd* addSurface = static_cast<ASTOpInstanceAdd*>(yes.get());
-
-                        bool ended = false;
-                        while (!ended)
-                        {
-                            ended = true;
-							Ptr<ASTOp> base = addSurface->instance.child();
-                            if ( base && base->GetOpType()==EOpType::IN_CONDITIONAL )
-                            {
-								ASTOpConditional* bottomConditional = static_cast<ASTOpConditional*>(base.get());
-
-                                // Are the two conditions exclusive?
-                                bool conditionaAreExclusive = false;
-
-                                ASTOpList facts;
-                                facts.Add(topConditional->condition.child());
-
-                                // Check if the child condition has a value with the current facts
-                                Ptr<ASTOp> pChildCond = bottomConditional->condition.child();
-                                ASTOp::FBoolEvalResult result;
-                                {
-                                    //MUTABLE_CPUPROFILER_SCOPE(EvaluateBool);
-                                    result = pChildCond->EvaluateBool( facts );
-                                }
-                                if ( result==ASTOp::BET_FALSE )
-                                {
-                                    conditionaAreExclusive = true;
-                                }
-
-                                if (conditionaAreExclusive)
-                                {
-                                    if (addSurface->GetParentCount()==1)
-                                    {
-                                        // Directly modify the instruction to skip the impossible child option
-                                        addSurface->instance = bottomConditional->no.child();
-                                    }
-                                    else
-                                    {
-                                        // Other parents may not impose the same condition that allows the optimisation.
-                                        Ptr<ASTOpInstanceAdd> newAddSurface = UE::Mutable::Private::Clone<ASTOpInstanceAdd>(addSurface);
-                                        newAddSurface->instance = bottomConditional->no.child();
-                                        topConditional->yes = newAddSurface;
-										addSurface = newAddSurface.get();
-                                    }
-
-                                    bModified = true;
-                                    ended = false;
-                                }
-                            }
-                        }
-                    }
-                }
-
-                else if (n->GetOpType()==EOpType::ME_CONDITIONAL)
+            	if (Op->IsConditional())
                 {
-					ASTOpConditional* topConditional = static_cast<ASTOpConditional*>(n.get());
-                    Ptr<ASTOp> yes = topConditional->yes.child();
-                    if ( yes && yes->GetOpType()==EOpType::ME_MERGE )
+					ASTOpConditional* TopConditional = static_cast<ASTOpConditional*>(Op.get());
+                    Ptr<ASTOp> Base = TopConditional->yes.child();
+                    if (Base)
                     {
-                        Ptr<ASTOpMeshMerge> Merge = static_cast<ASTOpMeshMerge*>(yes.get());
-
-                        bool ended = false;
-                        while (!ended)
+                        bool bEnd = false;
+                        while (!bEnd)
                         {
-                            ended = true;
-							Ptr<ASTOp> base = Merge->Base.child();
-                            if ( base && base->GetOpType()==EOpType::ME_CONDITIONAL )
-                            {
-								ASTOpConditional* bottomConditional = static_cast<ASTOpConditional*>(base.get());
-
-                                // Are the two conditions exclusive?
-                                bool conditionaAreExclusive = false;
-
-                                ASTOpList facts;
-                                facts.Add(topConditional->condition.child());
-
-                                // Check if the child condition has a value with the current facts
-								Ptr<ASTOp> pChildCond = bottomConditional->condition.child();
-                                ASTOp::FBoolEvalResult result = pChildCond->EvaluateBool( facts );
-                                if ( result==ASTOp::BET_FALSE )
-                                {
-                                    conditionaAreExclusive = true;
-                                }
-
-                                if (conditionaAreExclusive)
-                                {
-                                    if (Merge->GetParentCount()==1)
-                                    {
-                                        // Directly modify the instruction to skip the impossible child option
-										Merge->Base = bottomConditional->no.child();
-                                    }
-                                    else
-                                    {
-                                        // Other parents may not impose the same condition that allows
-                                        // the optimisation.
-                                        Ptr<ASTOpMeshMerge> NewMerge = UE::Mutable::Private::Clone<ASTOpMeshMerge>(Merge);
-										NewMerge->Base = bottomConditional->no.child();
-                                        topConditional->yes = NewMerge;
-                                        Merge = NewMerge;
-                                    }
-
-                                    bModified = true;
-                                    ended = false;
-                                }
-                            }
+                            bEnd = true;
+                        	Base->ForEachChild([&](ASTChild& Child)
+                        	{
+                        		if (!Child || !Child->IsConditional())
+	                            {
+		                            return;
+	                            }
+                        			
+		                        ASTOpConditional* BottomConditional = static_cast<ASTOpConditional*>(Child.child().get());
+
+		                        ASTOpList Facts;
+		                        Facts.Add(TopConditional->condition.child());
+									
+		                        Ptr<ASTOp> BottomCondition = BottomConditional->condition.child();
+		                        ASTOp::FBoolEvalResult Result = BottomCondition->EvaluateBool( Facts );
+                        		
+		                        const bool bIsConditionalExclusive = Result == ASTOp::BET_FALSE;
+		                        if (bIsConditionalExclusive)
+		                        {
+			                        if (Base->GetParentCount() == 1)
+			                        {
+				                        // Directly modify the instruction to skip the impossible child option
+				                        Child = BottomConditional->no.child();
+			                        }
+			                        else
+			                        {
+				                        // Other parents may not impose the same condition that allows the optimisation.
+				                        Ptr<ASTOp> NewBase = UE::Mutable::Private::Clone<ASTOp>(Base);
+			                        	ASTOp::Replace(Child.child(), BottomConditional->no.child());
+				                        TopConditional->yes = NewBase;
+				                        Base = NewBase.get();
+			                        }
+
+			                        bModified = true;
+			                        bEnd = false;
+		                        }
+                        	});
                         }
                     }
                 }
 
-
-                return bRecurse;
+                return true;
             });
         }
 
diff --git a/Engine/Plugins/Mutable/Source/MutableTools/Private/MuT/Compiler.cpp b/Engine/Plugins/Mutable/Source/MutableTools/Private/MuT/Compiler.cpp
--- a/Engine/Plugins/Mutable/Source/MutableTools/Private/MuT/Compiler.cpp
+++ b/Engine/Plugins/Mutable/Source/MutableTools/Private/MuT/Compiler.cpp
@@ -162,7 +162,9 @@ namespace UE::Mutable::Private
 		{ true,		{ 0, 0, 0, 0, 0, 0, 0, 0, 0, 0, 0, 0, 0, 0, 0, 0, 0, 0, 0, 0, 0, 0, 0, 0, 0, 0, 0, 0, 0 }		},	// ME_SKELETALMESH_BREAK
 
 		{ false,	{ 0, 0, 0, 0, 0, 0, 0, 0, 0, 0, 0, 0, 0, 0, 0, 0, 0, 0, 0, 0, 0, 0, 0, 0, 0, 0, 0, 0, 0 }		},	// SK_PARAMETER
-
+		{ false,	{ 0, 0, 0, 0, 0, 0, 0, 0, 0, 0, 0, 0, 0, 0, 0, 0, 0, 0, 0, 0, 0, 0, 0, 0, 0, 0, 0, 0, 0 }		},	// SK_CONDITIONAL
+		{ false,	{ 0, 0, 0, 0, 0, 0, 0, 0, 0, 0, 0, 0, 0, 0, 0, 0, 0, 0, 0, 0, 0, 0, 0, 0, 0, 0, 0, 0, 0 }		},	// SK_SWITCH
+    	
 		{ false,	{ 0, 0, 0, 0, 0, 0, 0, 0, 0, 0, 0, 0, 0, 0, 0, 0, 0, 0, 0, 0, 0, 0, 0, 0, 0, 0, 0, 0, 0 }		},	// IN_CONDITIONAL
 		{ false,	{ 0, 0, 0, 0, 0, 0, 0, 0, 0, 0, 0, 0, 0, 0, 0, 0, 0, 0, 0, 0, 0, 0, 0, 0, 0, 0, 0, 0, 0 }		},	// IN_SWITCH
 		{ false,	{ 0, 0, 0, 0, 0, 0, 0, 0, 0, 0, 0, 0, 0, 0, 0, 0, 0, 0, 0, 0, 0, 0, 0, 0, 0, 0, 0, 0, 0 }		},	// IN_ADDMESH
@@ -202,6 +204,7 @@ namespace UE::Mutable::Private
 
 		{ false,	{ 0, 0, 0, 0, 0, 0, 0, 0, 0, 0, 0, 0, 0, 0, 0, 0, 0, 0, 0, 0, 0, 0, 0, 0, 0, 0, 0, 0, 0 }		},	// IS_PARAMETER
     	{ true,		{ 0, 0, 0, 0, 0, 0, 0, 0, 0, 0, 0, 0, 0, 0, 0, 0, 0, 0, 0, 0, 0, 0, 0, 0, 0, 0, 0, 0, 0 }		},	// IS_EXTERNAL
+		{ false,	{ 0, 0, 0, 0, 0, 0, 0, 0, 0, 0, 0, 0, 0, 0, 0, 0, 0, 0, 0, 0, 0, 0, 0, 0, 0, 0, 0, 0, 0 }		},	// IS_CONDITIONAL
 		{ false,	{ 0, 0, 0, 0, 0, 0, 0, 0, 0, 0, 0, 0, 0, 0, 0, 0, 0, 0, 0, 0, 0, 0, 0, 0, 0, 0, 0, 0, 0 }		},	// IS_SWITCH
 	};
 
diff --git a/Engine/Plugins/Mutable/Source/MutableTools/Private/MuT/ErrorLog.cpp b/Engine/Plugins/Mutable/Source/MutableTools/Private/MuT/ErrorLog.cpp
--- a/Engine/Plugins/Mutable/Source/MutableTools/Private/MuT/ErrorLog.cpp
+++ b/Engine/Plugins/Mutable/Source/MutableTools/Private/MuT/ErrorLog.cpp
@@ -299,6 +299,8 @@ namespace UE::Mutable::Private
 		TEXT("ME_SKELETALMESH_BREAK"),
 
 		TEXT("SK_PARAMETER"),
+		TEXT("SK_CONDITIONAL"),
+		TEXT("SK_SWITCH"),
 
 		TEXT("IN_CONDITIONAL"),
 		TEXT("IN_SWITCH"),
@@ -339,6 +341,7 @@ namespace UE::Mutable::Private
 
 		TEXT("IS_PARAMETER"),
 		TEXT("IS_EXTERNAL"),
+		TEXT("IS_CONDITIONAL"),
 		TEXT("IS_SWITCH"),
 	};
 
