[Mutable] Better Group Type option names.

#Synced-CL 51880525

diff --git a/Engine/Plugins/Mutable/Source/CustomizableObject/Public/MuCO/CustomizableObjectParameterTypeDefinitions.h b/Engine/Plugins/Mutable/Source/CustomizableObject/Public/MuCO/CustomizableObjectParameterTypeDefinitions.h
--- a/Engine/Plugins/Mutable/Source/CustomizableObject/Public/MuCO/CustomizableObjectParameterTypeDefinitions.h
+++ b/Engine/Plugins/Mutable/Source/CustomizableObject/Public/MuCO/CustomizableObjectParameterTypeDefinitions.h
@@ -32,8 +32,8 @@ UENUM(BlueprintType)
 enum class ECustomizableObjectGroupType : uint8
 {
 	COGT_TOGGLE			UMETA(DisplayName = "Toggle"),
-	COGT_ALL			UMETA(DisplayName = "All Options"),
-	COGT_ONE			UMETA(DisplayName = "At least one Option"),
+	COGT_ALL			UMETA(DisplayName = "All"),
+	COGT_ONE			UMETA(DisplayName = "One"),
 	COGT_ONE_OR_NONE    UMETA(DisplayName = "One or None")
 };
 
