[Mutable] Prevent automatic merge of material slots with the same name and different materials.
Remove incorrect log from runtime mesh conversion.

#rnx

#Synced-CL 52057193

diff --git a/Engine/Plugins/Mutable/Source/CustomizableObject/Private/MuCO/CustomizableObjectInstance.cpp b/Engine/Plugins/Mutable/Source/CustomizableObject/Private/MuCO/CustomizableObjectInstance.cpp
--- a/Engine/Plugins/Mutable/Source/CustomizableObject/Private/MuCO/CustomizableObjectInstance.cpp
+++ b/Engine/Plugins/Mutable/Source/CustomizableObject/Private/MuCO/CustomizableObjectInstance.cpp
@@ -4820,7 +4820,9 @@ void UCustomizableInstancePrivate::BuildResources(const TSharedRef<FUpdateContex
 					MaterialSlot.MaterialSlotName = FoundSurfaceMetadata->MaterialSlotName;
 
 #if WITH_EDITOR
-					MaterialSlot.ImportedMaterialSlotName = FoundSurfaceMetadata->MaterialSlotName;
+					// Unique to avoid automatic merge of sections using same name when building the mesh in the editor.
+					FString ImportedMaterialSlotName = FoundSurfaceMetadata->MaterialSlotName.ToString() + FString::FromInt(MaterialSlotIndex);
+					MaterialSlot.ImportedMaterialSlotName = *ImportedMaterialSlotName;
 #endif
 				}
 
diff --git a/Engine/Plugins/Mutable/Source/CustomizableObject/Private/MuCO/UnrealConversionUtils.cpp b/Engine/Plugins/Mutable/Source/CustomizableObject/Private/MuCO/UnrealConversionUtils.cpp
--- a/Engine/Plugins/Mutable/Source/CustomizableObject/Private/MuCO/UnrealConversionUtils.cpp
+++ b/Engine/Plugins/Mutable/Source/CustomizableObject/Private/MuCO/UnrealConversionUtils.cpp
@@ -1504,11 +1504,6 @@ namespace UnrealConversionUtils
 
 				if (!Data->LODRenderData.IsValidIndex(LODIndex))
 				{
-					// This may happen if we are requesting a LOD that is not present in the provided mesh.
-					// TODO: an empty mesh is returned, but should we use the last LOD instead?
-					UE_LOGF(LogMutable, Warning, "Trying to convert mesh [%ls] LOD %d but it only has %d LODs", 
-						*SkeletalMesh->GetName(), LODIndex, Data->LODRenderData.Num() );
-					
 					return;
 				}
 
