[Mutable] Fix hang when compiling with External Operations.

#Synced-CL 51880525

diff --git a/Engine/Plugins/Mutable/Source/CustomizableObject/Internal/MuCO/MutableStreamableManager.h b/Engine/Plugins/Mutable/Source/CustomizableObject/Internal/MuCO/MutableStreamableManager.h
--- a/Engine/Plugins/Mutable/Source/CustomizableObject/Internal/MuCO/MutableStreamableManager.h
+++ b/Engine/Plugins/Mutable/Source/CustomizableObject/Internal/MuCO/MutableStreamableManager.h
@@ -72,7 +72,7 @@ public:
 		return StreamableManager.RequestSyncLoad(TargetsToStream, bManageActiveHandle, MoveTemp(DebugName), MoveTemp(Location));
 	}
 
-	void Tick(bool bBlocking);
+	CUSTOMIZABLEOBJECT_API void Tick(bool bBlocking);
 	
 	int32 GetRemainingWork() const;
 	
diff --git a/Engine/Plugins/Mutable/Source/CustomizableObjectEditor/Private/MuCOE/CustomizableObjectCompileRunnable.cpp b/Engine/Plugins/Mutable/Source/CustomizableObjectEditor/Private/MuCOE/CustomizableObjectCompileRunnable.cpp
--- a/Engine/Plugins/Mutable/Source/CustomizableObjectEditor/Private/MuCOE/CustomizableObjectCompileRunnable.cpp
+++ b/Engine/Plugins/Mutable/Source/CustomizableObjectEditor/Private/MuCOE/CustomizableObjectCompileRunnable.cpp
@@ -320,17 +320,20 @@ uint32 FCustomizableObjectCompileRunnable::Run()
 		CompilerOptions->SetDisableMaterialGeneration(true);
 	}
 
-	TFunction<void()> WaitCallback = [WeakCompiler=WeakCompiler]()
+	TFunction<void()> WaitCallback = [WeakCompiler = WeakCompiler]()
+	{
+		TSharedPtr<FCustomizableObjectCompiler> Compiler = WeakCompiler.Pin();
+		if (Compiler)
 		{
-			if (IsInGameThread())
-			{
-				TSharedPtr<FCustomizableObjectCompiler> Compiler = WeakCompiler.Pin();
-				if (Compiler)
-				{
-					Compiler->Tick();
-				}
-			}
-		};
+			Compiler->Tick(true);
+		}
+			
+		UCustomizableObjectSystem* System = UCustomizableObjectSystem::GetInstance();
+		if (System)
+		{
+			System->GetPrivate()->StreamableManager->Tick(true);
+		}
+	};
 
 	UE::Mutable::Private::Ptr<UE::Mutable::Private::Compiler> Compiler = new UE::Mutable::Private::Compiler(CompilerOptions, WaitCallback);
 
diff --git a/Engine/Plugins/Mutable/Source/CustomizableObjectEditor/Private/MuCOE/CustomizableObjectCompiler.h b/Engine/Plugins/Mutable/Source/CustomizableObjectEditor/Private/MuCOE/CustomizableObjectCompiler.h
--- a/Engine/Plugins/Mutable/Source/CustomizableObjectEditor/Private/MuCOE/CustomizableObjectCompiler.h
+++ b/Engine/Plugins/Mutable/Source/CustomizableObjectEditor/Private/MuCOE/CustomizableObjectCompiler.h
@@ -40,7 +40,7 @@ class FCustomizableObjectCompiler : public TSharedFromThis<FCustomizableObjectCo
 {
 public:
 	/** Check for pending compilation process. Returns true if an object has been updated. */
-	void Tick(bool bBlocking = false);
+	void Tick(bool bBlocking);
 	
 	int32 GetRemainingWork() const;
 	
diff --git a/Engine/Plugins/Mutable/Source/MutableTools/Private/MuT/CodeGenerator.cpp b/Engine/Plugins/Mutable/Source/MutableTools/Private/MuT/CodeGenerator.cpp
--- a/Engine/Plugins/Mutable/Source/MutableTools/Private/MuT/CodeGenerator.cpp
+++ b/Engine/Plugins/Mutable/Source/MutableTools/Private/MuT/CodeGenerator.cpp
@@ -141,6 +141,24 @@ namespace UE::Mutable::Private
 	}
 
 
+	// This function is a very bad idea. Someday we will hang due to oversubscription.
+	// We need to convert all recursive CodeGenerator to Tasks and avoid doing Wait/GetResult.
+	void CodeGenerator::WaitTask(Tasks::FTask Task)
+	{
+		if (IsInGameThread())
+		{
+			while (!Task.IsCompleted())
+			{
+				WaitCallback();
+			}
+		}
+		else
+		{
+			Task.Wait();
+		}
+	}
+
+
 	void CodeGenerator::GenerateGeneric(const FGenericGenerationOptions& Options, FGenericGenerationResult& OutResult, const Ptr<const Node> InNode )
 	{
 		if (!InNode)
@@ -662,15 +680,7 @@ namespace UE::Mutable::Private
 				return Result;
 			});
 
-		// Sync point: we currently don't support task-based generation beyond components so we wait here.
-		// Otherwise we could return the component task to be chained with higher level tasks.
-		if (WaitCallback.IsSet())
-		{
-			while (!ComponentTask.IsCompleted())
-			{
-				WaitCallback();
-			}
-		}
+		WaitTask(ComponentTask);
 		Result = ComponentTask.GetResult();
 	}
 
diff --git a/Engine/Plugins/Mutable/Source/MutableTools/Private/MuT/CodeGenerator.h b/Engine/Plugins/Mutable/Source/MutableTools/Private/MuT/CodeGenerator.h
--- a/Engine/Plugins/Mutable/Source/MutableTools/Private/MuT/CodeGenerator.h
+++ b/Engine/Plugins/Mutable/Source/MutableTools/Private/MuT/CodeGenerator.h
@@ -154,7 +154,7 @@ namespace UE::Mutable::Private
         //! Data will be stored in States
         void GenerateRoot( const Ptr<const Node> );
 
-	public:
+    	void WaitTask(Tasks::FTask Task);
 
 		/** This function will be called repeatedly in case a synchronization between threads is necessary. 
 		* It can be called from any thread.
diff --git a/Engine/Plugins/Mutable/Source/MutableTools/Private/MuT/CodeGenerator_External.cpp b/Engine/Plugins/Mutable/Source/MutableTools/Private/MuT/CodeGenerator_External.cpp
--- a/Engine/Plugins/Mutable/Source/MutableTools/Private/MuT/CodeGenerator_External.cpp
+++ b/Engine/Plugins/Mutable/Source/MutableTools/Private/MuT/CodeGenerator_External.cpp
@@ -125,15 +125,8 @@ namespace UE::Mutable::Private
 					UE::Tasks::MakeCompletedTask<FMeshGenerationDynamicOptions>();
     					
 				FMeshTask Task = GenerateMesh(MeshStaticOptions, MeshDynamicOptions, static_pointer_cast<NodeMesh>(Input));
-				
-				if (WaitCallback.IsSet())
-				{
-					while (!Task.IsCompleted())
-					{
-						WaitCallback();
-					}
-				}
-				
+
+				WaitTask(Task);
 				FMeshGenerationResult MeshResult = Task.GetResult();
 
 				External->Inputs.Add(ASTChild(External.get(), MeshResult.MeshOp));
diff --git a/Engine/Plugins/Mutable/Source/MutableTools/Private/MuT/CodeGenerator_Image.cpp b/Engine/Plugins/Mutable/Source/MutableTools/Private/MuT/CodeGenerator_Image.cpp
--- a/Engine/Plugins/Mutable/Source/MutableTools/Private/MuT/CodeGenerator_Image.cpp
+++ b/Engine/Plugins/Mutable/Source/MutableTools/Private/MuT/CodeGenerator_Image.cpp
@@ -1429,13 +1429,7 @@ namespace UE::Mutable::Private
 
 			FMeshTask MeshTask = GenerateMesh(MeshStaticOptions, UE::Tasks::MakeCompletedTask<FMeshGenerationDynamicOptions>(MeshOptions), Node->Mesh );
 			// This forces a wait to sync here. When images are also generated in tasks, it can be turned into a prerequisite instead.
-			if (WaitCallback.IsSet())
-			{
-				while (!MeshTask.IsCompleted())
-				{
-					WaitCallback();
-				}
-			}
+    		WaitTask(MeshTask);
 			FMeshGenerationResult MeshResult = MeshTask.GetResult();
 
 			// Match the block id of the block we are generating with the id that resulted in the generated mesh			
diff --git a/Engine/Plugins/Mutable/Source/MutableTools/Private/MuT/CodeGenerator_SkeletalMesh.cpp b/Engine/Plugins/Mutable/Source/MutableTools/Private/MuT/CodeGenerator_SkeletalMesh.cpp
--- a/Engine/Plugins/Mutable/Source/MutableTools/Private/MuT/CodeGenerator_SkeletalMesh.cpp
+++ b/Engine/Plugins/Mutable/Source/MutableTools/Private/MuT/CodeGenerator_SkeletalMesh.cpp
@@ -78,13 +78,7 @@ namespace UE::Mutable::Private
     	Ptr<ASTOpAddLOD> LODsOp = new ASTOpAddLOD();
     	for (FLODTask& LODTask : LODTasks)
     	{
-    		if (WaitCallback.IsSet())
-			{
-				while (!LODTask.IsCompleted())
-				{
-					WaitCallback();
-				}
-			}
+    		WaitTask(LODTask);
 			
     		FGenericGenerationResult LODResult = LODTask.GetResult();
     		LODsOp->lods.Emplace(LODsOp, LODResult.Op);
