[Mutable] Fix morphs not working for parametric meshes on cooked builds.

#rnx
#Synced-CL 52099427

diff --git a/Engine/Plugins/Mutable/Source/CustomizableObject/Private/MuCO/UnrealConversionUtils.cpp b/Engine/Plugins/Mutable/Source/CustomizableObject/Private/MuCO/UnrealConversionUtils.cpp
--- a/Engine/Plugins/Mutable/Source/CustomizableObject/Private/MuCO/UnrealConversionUtils.cpp
+++ b/Engine/Plugins/Mutable/Source/CustomizableObject/Private/MuCO/UnrealConversionUtils.cpp
@@ -95,9 +95,9 @@ namespace UnrealConversionUtils
 		/**
 		 * Initializes the skin vertex buffers object provided with the data found on the mutable buffers
 		 * @param OutVertexWeightBuffer - The Unreal's vertex buffers container to be updated with the mutable data.
-		   * @param NumVertices - The amount of vertices on the buffer
+		 * @param NumVertices - The amount of vertices on the buffer
 		 * @param NumBones - The amount of bones to use to init the skin weights buffer
-		   * @param NumBoneInfluences - The amount of bone influences on the buffer
+		 * @param NumBoneInfluences - The amount of bone influences on the buffer
 		 * @param bNeedCPUAccess - Determines if CPU access is required
 		 * @param InMutableData - Mutable data buffer
 		 */
@@ -1199,6 +1199,8 @@ namespace UnrealConversionUtils
 		
 			if (OptionalMorphTargetVertexBuffer)
 			{
+				check(OptionalMorphTargetVertexBuffer->IsMorphCPUDataValid());
+
 				const int32 NumMorphs = OptionalMorphTargetVertexBuffer->GetNumMorphs();
 				Morph.MaximumValuePerMorph.SetNumUninitialized(NumMorphs);
 				Morph.MinimumValuePerMorph.SetNumUninitialized(NumMorphs);
@@ -1959,11 +1961,17 @@ namespace UnrealConversionUtils
 
 				if (!bIgnoreMorphs)
 				{
+					FMorphTargetVertexInfoBuffers* MorphTargetBuffers = nullptr;
+					if (Context->DataLOD->MorphTargetVertexInfoBuffers.IsMorphCPUDataValid())
+					{
+						MorphTargetBuffers = &Context->DataLOD->MorphTargetVertexInfoBuffers;
+					}
+
 					SetMorphData(
 						Result.Get(), SkeletalMesh.Get(),
 						LODIndex, SectionIndex,
 						(int32)Section.BaseVertexIndex, (int32)Section.NumVertices, (int32)Context->OriginalLOD->GetNumVertices(),
-						nullptr);
+						MorphTargetBuffers);
 				}
 
 				Result->EnsureSurfaceData();
