[Mutable] Fix CO Instances not updating.
Move CustomizableSkeletalComponent properties to private.
Expose CustomizableObjectInstanceUsage properties to blueprint and editor.
Optimize CustomizableObjectSystem tick game thread

#Synced-CL 52127744

diff --git a/Engine/Plugins/Experimental/MutablePopulation/Source/CustomizableObjectPopulationEditor/Private/MuCOPE/CustomizableObjectPopulationClassEditor.cpp b/Engine/Plugins/Experimental/MutablePopulation/Source/CustomizableObjectPopulationEditor/Private/MuCOPE/CustomizableObjectPopulationClassEditor.cpp
--- a/Engine/Plugins/Experimental/MutablePopulation/Source/CustomizableObjectPopulationEditor/Private/MuCOPE/CustomizableObjectPopulationClassEditor.cpp
+++ b/Engine/Plugins/Experimental/MutablePopulation/Source/CustomizableObjectPopulationEditor/Private/MuCOPE/CustomizableObjectPopulationClassEditor.cpp
@@ -797,7 +797,7 @@ void FCustomizableObjectPopulationClassEditor::TestPopulationClass()
 							ViewportInstances[i]->SetBuildParameterRelevancy(true);
 							ViewportInstances[i]->UpdateSkeletalMeshAsync(true, true);
 
-							PreviewCustomizableSkeletalComponent->CustomizableObjectInstance = ViewportInstances[i];
+							PreviewCustomizableSkeletalComponent->SetCustomizableObjectInstance(ViewportInstances[i]);
 							// TODO: UE-274335. We want this code to be able to operate with multiple components, not just the first one
 							PreviewCustomizableSkeletalComponent->SetComponentName(CustomizableObject->GetComponentName(0));
 							PreviewCustomizableSkeletalComponent->AttachToComponent(SkeletalMeshComponent, FAttachmentTransformRules::KeepRelativeTransform);
diff --git a/Engine/Plugins/Experimental/MutablePopulation/Source/CustomizableObjectPopulationEditor/Private/MuCOPE/CustomizableObjectPopulationEditor.cpp b/Engine/Plugins/Experimental/MutablePopulation/Source/CustomizableObjectPopulationEditor/Private/MuCOPE/CustomizableObjectPopulationEditor.cpp
--- a/Engine/Plugins/Experimental/MutablePopulation/Source/CustomizableObjectPopulationEditor/Private/MuCOPE/CustomizableObjectPopulationEditor.cpp
+++ b/Engine/Plugins/Experimental/MutablePopulation/Source/CustomizableObjectPopulationEditor/Private/MuCOPE/CustomizableObjectPopulationEditor.cpp
@@ -478,7 +478,7 @@ void FCustomizableObjectPopulationEditor::TestPopulation()
 					ViewportInstances[i]->SetBuildParameterRelevancy(true);
 					ViewportInstances[i]->UpdateSkeletalMeshAsync(true, true);
 
-					PreviewCustomizableSkeletalComponent->CustomizableObjectInstance = ViewportInstances[i];
+					PreviewCustomizableSkeletalComponent->SetCustomizableObjectInstance(ViewportInstances[i]);
 					// TODO: UE-274335. We want this code to be able to operate with multiple components, not just the first one
 					PreviewCustomizableSkeletalComponent->SetComponentName(CustomizableObject->GetComponentName(0));
 					PreviewCustomizableSkeletalComponent->AttachToComponent(PreviewSkeletalMeshComponent, FAttachmentTransformRules::KeepRelativeTransform);
diff --git a/Engine/Plugins/Mutable/Source/CustomizableObject/Internal/MuCO/CustomizableObjectInstancePrivate.h b/Engine/Plugins/Mutable/Source/CustomizableObject/Internal/MuCO/CustomizableObjectInstancePrivate.h
--- a/Engine/Plugins/Mutable/Source/CustomizableObject/Internal/MuCO/CustomizableObjectInstancePrivate.h
+++ b/Engine/Plugins/Mutable/Source/CustomizableObject/Internal/MuCO/CustomizableObjectInstancePrivate.h
@@ -59,12 +59,8 @@ enum ECOInstanceFlags
 	// Update process
 	ReplacePhysicsAssets			= 1 << 1,	// Merge active PhysicsAssets and replace the base physics asset
 
-	// Update priorities
-	UsedByPlayerOrNearIt			= 1 << 2,	// The instance is used by the player or is near the player, used to give more priority to its updates
-	Discarded						= 1 << 3,
-
-	HasPendingHighPriorityUpdate	= 1 << 4,
-	HasPendingLowPriorityUpdate		= 1 << 5,
+	HasPendingHighPriorityUpdate	= 1 << 2,
+	HasPendingLowPriorityUpdate		= 1 << 3,
 };
 
 ENUM_CLASS_FLAGS(ECOInstanceFlags);
@@ -268,8 +264,8 @@ public:
 
 	UE_API UCustomizableObjectInstance* GetPublic() const;
 
-	// If any components are using this instance, they will store the min of their distances to the player here every frame for LOD purposes
-	float MinSquareDistFromComponentToPlayer = FLT_MAX;
+	void RegisterCustomizableObjectInstanceUsage(UCustomizableObjectInstanceUsage& Obj);
+	void UnregisterCustomizableObjectInstanceUsage(UCustomizableObjectInstanceUsage& Obj);
 	
 	UPROPERTY(Transient)
 	TArray<FCustomizableInstanceComponentData> ComponentsData;
@@ -291,17 +287,11 @@ public:
 	/** Status of the generated Skeletal Mesh. Not to be confused with the Update Result. */
 	ESkeletalMeshStatus SkeletalMeshStatus = ESkeletalMeshStatus::NotGenerated;
 
-	TMap<FString, bool> ParamNameToExpandedMap; // Used to check whether a mutable param is expanded in the editor to show its child params
-
 	bool bShowOnlyRuntimeParameters = true;
 	bool bShowOnlyRelevantParameters = true;
 	bool bShowUISections = false;
 	bool bShowUIThumbnails = false;
 
-	/** Automatic update required.
-	 * Set to true when a Customizable Object Instance Usage requires an automatic update (e.g., component reattached). */
-	bool bAllowAutomaticUpdate = false;
-
 #if WITH_EDITOR
 	/** Profile index the instance parameters are in and if the profile needs to be refreshed */
 	int32 SelectedProfileIndex = INDEX_NONE;
@@ -326,6 +316,8 @@ public:
 
 	/** Pointer to a instance of class that logs update messages to an editor widget. */
 	TSharedPtr<IMutableEditorLogger> UpdateLogger;
+
+	TArray<UCustomizableObjectInstanceUsage*> InstanceUsages;
 };
 
 
diff --git a/Engine/Plugins/Mutable/Source/CustomizableObject/Internal/MuCO/CustomizableObjectInstanceUsagePrivate.h b/Engine/Plugins/Mutable/Source/CustomizableObject/Internal/MuCO/CustomizableObjectInstanceUsagePrivate.h
--- a/Engine/Plugins/Mutable/Source/CustomizableObject/Internal/MuCO/CustomizableObjectInstanceUsagePrivate.h
+++ b/Engine/Plugins/Mutable/Source/CustomizableObject/Internal/MuCO/CustomizableObjectInstanceUsagePrivate.h
@@ -2,18 +2,12 @@
 
 #pragma once
 
-#include "Tickable.h"
-
 #include "CustomizableObjectInstanceUsagePrivate.generated.h"
 
 #define UE_API CUSTOMIZABLEOBJECT_API
 
-class UCustomizableObjectInstance;
 class UCustomizableObjectInstanceUsage;
-class UCustomizableSkeletalComponent;
-class UPhysicsAsset;
 class USkeletalMesh;
-class USkeletalMeshComponent;
 
 UCLASS(MinimalAPI)
 class UCustomizableObjectInstanceUsagePrivate : public UObject
@@ -30,9 +24,6 @@ public:
 
 	UE_API USkeletalMesh* GetAttachedSkeletalMesh() const;
 	
-	/** Return a valid pointer if this Usage is being used inside a UCustomizableSkeletalComponent. Nullptr if it is being used in Standalone mode. */
-	UE_API UCustomizableSkeletalComponent* GetCustomizableSkeletalComponent() const;
-	
 	UE_API UCustomizableObjectInstanceUsage* GetPublic();
 
 	UE_API const UCustomizableObjectInstanceUsage* GetPublic() const;
diff --git a/Engine/Plugins/Mutable/Source/CustomizableObject/Internal/MuCO/CustomizableObjectSystemPrivate.h b/Engine/Plugins/Mutable/Source/CustomizableObject/Internal/MuCO/CustomizableObjectSystemPrivate.h
--- a/Engine/Plugins/Mutable/Source/CustomizableObject/Internal/MuCO/CustomizableObjectSystemPrivate.h
+++ b/Engine/Plugins/Mutable/Source/CustomizableObject/Internal/MuCO/CustomizableObjectSystemPrivate.h
@@ -308,8 +308,8 @@ public:
 
 	TMap<UE::Mutable::Private::FImageId, TUniquePtr<FTexturePlatformData>> ImageToPlatformDataMap;
 	
-	/** If a InstanceUsage is in this set it means that its AttachParent has been modified (USkeletalMesh changed, UMaterial changed...). */
-	TSet<TWeakObjectPtr<UCustomizableObjectInstanceUsage>> AttachedParentUpdated;
+	/** If a InstanceUsage is in this array it means that its AttachParent has been modified (USkeletalMesh changed, UMaterial changed...). */
+	TArray<UCustomizableObjectInstanceUsage*> AttachedParentUpdated;
 
 	FInstanceUpdateDelegate UpdateCallback;
 	FInstanceUpdateNativeDelegate UpdateNativeCallback;
@@ -558,11 +558,13 @@ public:
 	bool bAutoCompileCommandletEnabled = false;
 
 	// mutable.GenerateInstancesWithinRange = true
-	void UpdateInstanceDistances();
-	bool IsDiscardedByDistance(const UCustomizableInstancePrivate* Instance) const;
+	void UpdateViewLocations();
+	bool IsDiscardedByDistance(const UCustomizableObjectInstance& Instance) const;
+	void DiscardInstance(UCustomizableObjectInstance& Instance);
 
 	// The list of actors that define a view radius (GenerationRange) around them.
 	TSet<TWeakObjectPtr<const AActor>> ViewCenters;
+	TArray<FVector> ViewLocations;
 
 	// Instance generation range when mutable.GenerateInstancesWithinRange is true
 	float GenerationRangeSquare = 2000.f * 2000.f;
@@ -589,12 +591,6 @@ public:
 	float LogStartedUpdateLast = 0;
 
 public:
-	static const TArray<UCustomizableObjectInstanceUsage*>& GetCustomizableObjectInstanceUsages()
-	{
-		return CustomizableObjectInstanceUsages;
-	}
-	static void RegisterCustomizableObjectInstanceUsage(UCustomizableObjectInstanceUsage& Obj);
-	static void UnregisterCustomizableObjectInstanceUsage(UCustomizableObjectInstanceUsage& Obj);
 
 	static const TArray<UCustomizableObjectInstance*>& GetCustomizableObjectInstances()
 	{
@@ -603,6 +599,19 @@ public:
 	static void RegisterCustomizableObjectInstance(UCustomizableObjectInstance& Obj);
 	static void UnregisterCustomizableObjectInstance(UCustomizableObjectInstance& Obj);
 
+
+	static void RegisterInstanceToGeneratedList(UCustomizableObjectInstance& Obj);
+	static void UnregisterInstancefromGeneratedList(UCustomizableObjectInstance& Obj);
+	static void RegisterInstanceToAutomaticUpdateList(UCustomizableObjectInstance& Obj);
+	static void UnregisterInstanceFromAutomaticUpdateList(UCustomizableObjectInstance& Obj);
+
+	void TickDiscardInstances();
+	void TickAutomaticUpdateInstances();
+
+	static void RegisterInstanceToPendinSetReferenceSkeletalMeshList(UCustomizableObjectInstance& Obj);
+	void TickPendingSetReferenceSkeletalMesh();
+
+
 	/**
 	 * Get to know if the settings used by the mutable syustem are optimzed for benchmarking operations or not
 	 * @return true if using benchmarking optimized settings, false otherwise.
@@ -622,8 +631,14 @@ private:
 	inline static bool bUseBenchmarkingSettings = false;
 
 	// Added to in PostInitProperties, removed from in BeginDestroy. No CDOs
-	static TArray<UCustomizableObjectInstanceUsage*> CustomizableObjectInstanceUsages;
 	static TArray<UCustomizableObjectInstance*> CustomizableObjectInstances;
+	static TArray<UCustomizableObjectInstance*> GeneratedInstances;
+	static TArray<UCustomizableObjectInstance*> AutomaticUpdateInstances;
+
+	static TArray<UCustomizableObjectInstance*> PendingSetReferenceSkeletalMesh;
+
+	int32 DiscardIndex = 0;
+	int32 AutomaticUpdateIndex = 0;
 };
 
 
diff --git a/Engine/Plugins/Mutable/Source/CustomizableObject/Internal/MuCO/CustomizableSkeletalComponentPrivate.h b/Engine/Plugins/Mutable/Source/CustomizableObject/Internal/MuCO/CustomizableSkeletalComponentPrivate.h
--- a/Engine/Plugins/Mutable/Source/CustomizableObject/Internal/MuCO/CustomizableSkeletalComponentPrivate.h
+++ b/Engine/Plugins/Mutable/Source/CustomizableObject/Internal/MuCO/CustomizableSkeletalComponentPrivate.h
@@ -6,11 +6,7 @@
 
 #define UE_API CUSTOMIZABLEOBJECT_API
 
-class UCustomizableSkeletalComponent;
 class UCustomizableObjectInstanceUsage;
-class UPhysicsAsset;
-class USkeletalMesh;
-
 
 UCLASS(MinimalAPI)
 class UCustomizableSkeletalComponentPrivate : public UObject
@@ -20,14 +16,6 @@ class UCustomizableSkeletalComponentPrivate : public UObject
 public:
 	UE_API UCustomizableSkeletalComponentPrivate();
 	
-	UE_API USkeletalMesh* GetSkeletalMesh() const;
-
-	UE_API USkeletalMesh* GetAttachedSkeletalMesh() const;
-	
-	UE_API UCustomizableSkeletalComponent* GetPublic();
-	
-	UE_API const UCustomizableSkeletalComponent* GetPublic() const;
-
 	UPROPERTY(Instanced)
 	TObjectPtr<UCustomizableObjectInstanceUsage> InstanceUsage;
 };
diff --git a/Engine/Plugins/Mutable/Source/CustomizableObject/Internal/MuCO/ICustomizableObjectEditorModule.h b/Engine/Plugins/Mutable/Source/CustomizableObject/Internal/MuCO/ICustomizableObjectEditorModule.h
--- a/Engine/Plugins/Mutable/Source/CustomizableObject/Internal/MuCO/ICustomizableObjectEditorModule.h
+++ b/Engine/Plugins/Mutable/Source/CustomizableObject/Internal/MuCO/ICustomizableObjectEditorModule.h
@@ -109,8 +109,6 @@ public:
 	 * If calling inside the Customizable Object Editor module, consider using ICustomizableObjectEditorModulePrivate::EnqueueCompileRequest. */
 	virtual void CompileCustomizableObject(UCustomizableObject& Object, const FCompileParams* Params, bool bSilent, bool bForce) = 0;
 
-	virtual void UpdateInstanceDistanceToLevelEditorCamera() = 0;
-
 	virtual void Tick(bool bBlocking) = 0;
 	
 	virtual int32 GetRemainingWork() const = 0;
diff --git a/Engine/Plugins/Mutable/Source/CustomizableObject/Private/MuCO/CustomizableObjectInstance.cpp b/Engine/Plugins/Mutable/Source/CustomizableObject/Private/MuCO/CustomizableObjectInstance.cpp
--- a/Engine/Plugins/Mutable/Source/CustomizableObject/Private/MuCO/CustomizableObjectInstance.cpp
+++ b/Engine/Plugins/Mutable/Source/CustomizableObject/Private/MuCO/CustomizableObjectInstance.cpp
@@ -67,12 +67,6 @@ namespace
 	#define REQUIRES_SINGLEUSE_FLAG_FOR_RUNTIME_TEXTURES !PLATFORM_DESKTOP
 #endif
 
-bool bDisableClothingPhysicsEditsPropagation = false;
-static FAutoConsoleVariableRef CVarDisableClothingPhysicsEditsPropagation(
-	TEXT("mutable.DisableClothingPhysicsEditsPropagation"),
-	bDisableClothingPhysicsEditsPropagation,
-	TEXT("If set to true, disables clothing physics edits propagation from the render mesh."),
-	ECVF_Default);
 
 bool bDisableNotifyComponentsOfTextureUpdates = false;
 static FAutoConsoleVariableRef CVarDisableNotifyComponentsOfTextureUpdates(
@@ -118,6 +112,12 @@ void UCustomizableInstancePrivate::InvalidateGeneratedData()
 	CommittedDescriptorHash = {};
 
 	ComponentsData.Empty();
+
+	if (UCustomizableObjectInstance* Instance = GetPublic())
+	{
+		UCustomizableObjectSystemPrivate::UnregisterInstancefromGeneratedList(*Instance);
+		UCustomizableObjectSystemPrivate::RegisterInstanceToAutomaticUpdateList(*Instance);
+	}
 }
 
 
@@ -145,25 +145,19 @@ FCustomizableInstanceComponentData& UCustomizableInstancePrivate::GetComponentDa
 UCustomizableObjectInstance::UCustomizableObjectInstance()
 {
 	SetFlags(RF_Transactional);
+	PrivateData = CreateDefaultSubobject<UCustomizableInstancePrivate>(TEXT("Private"));
 }
 
 
 void UCustomizableInstancePrivate::SetDescriptor(const FCustomizableObjectInstanceDescriptor& InDescriptor)
 {
-	UCustomizableObject* InCustomizableObject = InDescriptor.GetCustomizableObject();
-	const bool bCustomizableObjectChanged = GetPublic()->Descriptor.GetCustomizableObject() != InCustomizableObject;
-
 #if WITH_EDITOR
 	// Bind a lambda to the PostCompileDelegate and unbind from the previous object if any.
-	BindObjectDelegates(GetPublic()->GetCustomizableObject(), InCustomizableObject);
+	BindObjectDelegates(GetPublic()->GetCustomizableObject(), InDescriptor.GetCustomizableObject());
 #endif
 
 	GetPublic()->Descriptor = InDescriptor;
-
-	if (bCustomizableObjectChanged)
-	{
-		InvalidateGeneratedData();
-	}
+	InvalidateGeneratedData();
 }
 
 
@@ -271,18 +265,6 @@ void UCustomizableObjectInstance::PostInitProperties()
 {
 	Super::PostInitProperties();
 	UCustomizableObjectSystemPrivate::RegisterCustomizableObjectInstance(*this);
-
-	if (!HasAnyFlags(RF_ClassDefaultObject | RF_DefaultSubObject))
-	{
-		if (!PrivateData)
-		{
-			PrivateData = NewObject<UCustomizableInstancePrivate>(this, FName("Private"));
-		}
-		else if (PrivateData->GetOuter() != this)
-		{
-			PrivateData = Cast<UCustomizableInstancePrivate>(StaticDuplicateObject(PrivateData, this, FName("Private")));
-		}
-	}
 }
 
 
@@ -307,13 +289,6 @@ void UCustomizableObjectInstance::DestroyLiveUpdateInstance()
 }
 
 
-bool UCustomizableObjectInstance::IsReadyForFinishDestroy()
-{
-	//return ReleaseResourcesFence.IsFenceComplete();
-	return true;
-}
-
-
 void UCustomizableObjectInstance::Serialize(FArchive& Ar)
 {
 	Super::Serialize(Ar);
@@ -486,6 +461,37 @@ void UCustomizableObjectInstance::SetCurrentState(const FString& StateName)
 }
 
 
+void UCustomizableInstancePrivate::RegisterCustomizableObjectInstanceUsage(UCustomizableObjectInstanceUsage& Obj)
+{
+	if (!Obj.IsTemplate() && !InstanceUsages.Contains(&Obj))
+	{
+		InstanceUsages.Add(&Obj);
+
+		UCustomizableObjectInstance* Instance = GetPublic();
+		UCustomizableObjectSystemPrivate::RegisterInstanceToAutomaticUpdateList(*Instance);
+
+		if (!Obj.GetSkipSetReferenceSkeletalMesh())
+		{
+			UCustomizableObjectSystemPrivate::RegisterInstanceToPendinSetReferenceSkeletalMeshList(*Instance);
+		}
+	}
+}
+
+
+void UCustomizableInstancePrivate::UnregisterCustomizableObjectInstanceUsage(UCustomizableObjectInstanceUsage& Obj)
+{
+	if (!Obj.IsTemplate())
+	{
+		InstanceUsages.Remove(&Obj);
+
+		if (InstanceUsages.IsEmpty())
+		{
+			UCustomizableObjectSystemPrivate::UnregisterInstanceFromAutomaticUpdateList(*GetPublic());
+		}
+	}
+}
+
+
 bool UCustomizableInstancePrivate::IsParameterRelevant(uint32 ParameterIndex) const
 {
 	// This should have been precalculated in the last update if the appropriate flag in the instance was set.
@@ -1313,6 +1319,11 @@ UPhysicsAsset* GetOrBuildMainPhysicsAsset(
 	check(MutablePhysics);
 
 	UCustomizableObjectInstance* Instance = Context->Instance.Get();
+	if (!Instance)
+	{
+		return nullptr;
+	}
+
 	UCustomizableInstancePrivate* InstancePrivate = Instance->GetPrivate();
 	
 	UPhysicsAsset* Result = nullptr;
@@ -1616,14 +1627,14 @@ bool DoComponentsNeedUpdate(const TSharedRef<FUpdateContextPrivate>& Context, bo
 	MUTABLE_CPUPROFILER_SCOPE(UCustomizableInstancePrivate::DoComponentsNeedUpdate);
 
 	UCustomizableObjectInstance* Instance = Context->Instance.Get();
-	UCustomizableInstancePrivate* InstancePrivate = Instance->GetPrivate();
-	
-	UCustomizableObject* CustomizableObject = Instance->GetCustomizableObject();
-	
-	if (!CustomizableObject)
+	UCustomizableObject* CustomizableObject = Context->Object.Get();
+
+	if (!Instance || !CustomizableObject)
 	{
 		return false;
 	}
+
+	UCustomizableInstancePrivate* InstancePrivate = Instance->GetPrivate();
 	
 	bHasInvalidMesh = false;
 	bool bNeedsUpdate = false;
@@ -2685,13 +2696,8 @@ void UCustomizableInstancePrivate::ForceSetReferenceSkeletalMesh() const
 		return;
 	}
 	
-	for (UCustomizableObjectInstanceUsage* CustomizableObjectInstanceUsage : UCustomizableObjectSystemPrivate::GetCustomizableObjectInstanceUsages())
+	for (UCustomizableObjectInstanceUsage* CustomizableObjectInstanceUsage : InstanceUsages)
 	{
-		if (!IsValid(CustomizableObjectInstanceUsage) || CustomizableObjectInstanceUsage->GetCustomizableObjectInstance() != Instance)
-		{
-			continue;
-		}
-
 #if WITH_EDITOR
 		if (CustomizableObjectInstanceUsage->GetPrivate()->bIsNetModeDedicatedServer)
 		{
@@ -3151,6 +3157,8 @@ void InitSkeletalMeshData(const TSharedRef<FUpdateContextPrivate>& Context, USke
 	check(SkeletalMesh);
 
 	UCustomizableObjectInstance* Instance = Context->Instance.Get();
+	check(Instance);
+
 	UCustomizableInstancePrivate* InstancePrivate = Instance->GetPrivate();
 	
 	const FName ComponentName = Context->ComponentNames[Component->Id];
@@ -3375,10 +3383,8 @@ bool BuildSkeletonData(const TSharedRef<FUpdateContextPrivate>& Context, USkelet
 {
 	MUTABLE_CPUPROFILER_SCOPE(UCustomizableInstancePrivate::BuildSkeletonData);
 
-	UCustomizableObjectInstance* Instance = Context->Instance.Get();
-	UCustomizableInstancePrivate* InstancePrivate = Instance->GetPrivate();
-
-	UCustomizableObject* CustomizableObject = Instance->GetCustomizableObject();
+	UCustomizableObject* CustomizableObject = Context->Object.Get();
+	check(CustomizableObject);
 
 	const TSharedPtr<FInstanceUpdateData::FComponent>& Component = Context->InstanceUpdateData.Components[ComponentIndex];
 	if (!Component)
@@ -4063,14 +4069,13 @@ UE::Tasks::FTask LoadAdditionalAssetsAndData(const TSharedRef<FUpdateContextPriv
 	MUTABLE_CPUPROFILER_SCOPE(UCustomizableInstancePrivate::LoadAdditionalAssetsAndDataAsync);
 
 	UCustomizableObjectInstance* Instance = Context->Instance.Get();
-	UCustomizableInstancePrivate* InstancePrivate = Instance->GetPrivate();
-	
-	UCustomizableObject* CustomizableObject = Instance->GetCustomizableObject();
+	UCustomizableObject* CustomizableObject = Context->Object.Get();
 
-	const UModelResources& ModelResources = CustomizableObject->GetPrivate()->GetModelResourcesChecked();
-	const TSharedPtr<FModelStreamableBulkData>& ModelStreamableBulkData = CustomizableObject->GetPrivate()->GetModelStreamableBulkData();
+	check(Instance);
+	check(CustomizableObject);
 	
-	FMutableStreamRequest StreamRequest(ModelStreamableBulkData);
+	UCustomizableInstancePrivate* InstancePrivate = Instance->GetPrivate();
+	const UModelResources& ModelResources = CustomizableObject->GetPrivate()->GetModelResourcesChecked();
 
 	TArray<FSoftObjectPath> AssetsToStream;
 
@@ -4099,10 +4104,6 @@ UE::Tasks::FTask LoadAdditionalAssetsAndData(const TSharedRef<FUpdateContextPriv
 			FStreamableDelegate::CreateStatic(&AdditionalAssetsAsyncLoaded, Context, Event),
 			CVarMutableHighPriorityLoading.GetValueOnAnyThread() ? FStreamableManager::AsyncLoadHighPriority : FStreamableManager::DefaultAsyncLoadPriority);
 	}
-	
-	// Stream files
-	UE::Tasks::FTask StreamingTask = StreamRequest.Stream();
-	Prerequisites.Add(StreamingTask);
 
 	return UE::Tasks::Launch(TEXT("CaptureContext"), [Context]() {}, // Keep a reference to make sure allocated memory is always alive. TODO Probably not necessary since AdditionalAssetsAsyncLoaded already captures Context.
 	Prerequisites,
@@ -4867,13 +4868,8 @@ void UCustomizableInstancePrivate::BuildResources(const TSharedRef<FUpdateContex
 	{
 		MUTABLE_CPUPROFILER_SCOPE(BuildResources_RecreateRenderState);
 
-		for (UCustomizableObjectInstanceUsage* CustomizableObjectInstanceUsage : UCustomizableObjectSystemPrivate::GetCustomizableObjectInstanceUsages())
+		for (UCustomizableObjectInstanceUsage* CustomizableObjectInstanceUsage : InstanceUsages)
 		{
-			if (!IsValid(CustomizableObjectInstanceUsage) || CustomizableObjectInstanceUsage->GetCustomizableObjectInstance() != Public)
-			{
-				continue;
-			}
-
 #if WITH_EDITOR
 			if (CustomizableObjectInstanceUsage->GetPrivate()->bIsNetModeDedicatedServer)
 			{
diff --git a/Engine/Plugins/Mutable/Source/CustomizableObject/Private/MuCO/CustomizableObjectInstanceUsage.cpp b/Engine/Plugins/Mutable/Source/CustomizableObject/Private/MuCO/CustomizableObjectInstanceUsage.cpp
--- a/Engine/Plugins/Mutable/Source/CustomizableObject/Private/MuCO/CustomizableObjectInstanceUsage.cpp
+++ b/Engine/Plugins/Mutable/Source/CustomizableObject/Private/MuCO/CustomizableObjectInstanceUsage.cpp
@@ -3,41 +3,17 @@
 #include "MuCO/CustomizableObjectInstanceUsage.h"
 #include "MuCO/CustomizableObjectInstanceUsagePrivate.h"
 
-#include "MuCO/CustomizableSkeletalComponent.h"
+#include "MuCO/CustomizableObject.h"
 #include "MuCO/CustomizableObjectSystem.h"
 #include "MuCO/CustomizableObjectInstancePrivate.h"
-#include "MuCO/UnrealPortabilityHelpers.h"
-#include "MuCO/CustomizableSkeletalComponentPrivate.h"
-#include "AnimationRuntime.h"
 #include "Components/SkeletalMeshComponent.h"
 #include "Engine/SkeletalMesh.h"
-#include "Engine/SkinnedAssetCommon.h"
-#include "Engine/World.h"
-#include "GameFramework/Pawn.h"
-#include "Materials/MaterialInstanceDynamic.h"
-#include "MuCO/CustomizableObject.h"
-#include "MuCO/ICustomizableObjectModule.h"
-#include "Rendering/SkeletalMeshRenderData.h"
-#include "UObject/ObjectSaveContext.h"
-#include "Stats/Stats.h"
 
 #include UE_INLINE_GENERATED_CPP_BY_NAME(CustomizableObjectInstanceUsage)
 
 
 void UCustomizableObjectInstanceUsagePrivate::Callbacks()
 {
-	
-	if (UCustomizableSkeletalComponent* CustomizableSkeletalComponent = GetCustomizableSkeletalComponent())
-	{
-		CustomizableSkeletalComponent->UpdatedDelegate.ExecuteIfBound();
-
-		if (GetPublic()->UpdatedDelegate.IsBound() && CustomizableSkeletalComponent->UpdatedDelegate.IsBound())
-		{
-			UE_LOGF(LogMutable, Error, "The UpdatedDelegate is bound both in the UCustomizableObjectInstanceUsage and in its parent CustomizableSkeletalComponent. Only one should be bound.");
-			ensure(false);
-		}
-	}
-	
 	GetPublic()->UpdatedDelegate.ExecuteIfBound();
 }
 
@@ -47,89 +23,109 @@ UCustomizableObjectInstanceUsage::UCustomizableObjectInstanceUsage()
 	Private = CreateDefaultSubobject<UCustomizableObjectInstanceUsagePrivate>(FName("Private"));
 }
 
-void UCustomizableObjectInstanceUsage::PostInitProperties()
-{
-	Super::PostInitProperties();
-	UCustomizableObjectSystemPrivate::RegisterCustomizableObjectInstanceUsage(*this);
-}
 
-void UCustomizableObjectInstanceUsage::BeginDestroy()
+void UCustomizableObjectInstanceUsage::PostDuplicate(bool bDuplicateForPIE)
 {
-	UCustomizableObjectSystemPrivate::UnregisterCustomizableObjectInstanceUsage(*this);
-	Super::BeginDestroy();
+	Super::PostDuplicate(bDuplicateForPIE);
+
+	if (UsedCustomizableObjectInstance)
+	{
+		UsedCustomizableObjectInstance->GetPrivate()->RegisterCustomizableObjectInstanceUsage(*this);
+	}
 }
 
-void UCustomizableObjectInstanceUsage::SetCustomizableObjectInstance(UCustomizableObjectInstance* CustomizableObjectInstance)
+void UCustomizableObjectInstanceUsage::PostLoad()
 {
-	if (UCustomizableSkeletalComponent* CustomizableSkeletalComponent = GetPrivate()->GetCustomizableSkeletalComponent())
+	Super::PostLoad();
+
+	if (UsedCustomizableObjectInstance)
 	{
-		CustomizableSkeletalComponent->CustomizableObjectInstance = CustomizableObjectInstance;
+		UsedCustomizableObjectInstance->GetPrivate()->RegisterCustomizableObjectInstanceUsage(*this);
 	}
-	else
+}
+
+void UCustomizableObjectInstanceUsage::BeginDestroy()
+{
+	if (UsedCustomizableObjectInstance)
 	{
-		UsedCustomizableObjectInstance = CustomizableObjectInstance;
+		UsedCustomizableObjectInstance->GetPrivate()->UnregisterCustomizableObjectInstanceUsage(*this);
 	}
-}
 
+	Super::BeginDestroy();
+}
 
-UCustomizableObjectInstance* UCustomizableObjectInstanceUsage::GetCustomizableObjectInstance() const
+#if WITH_EDITOR
+void UCustomizableObjectInstanceUsage::PostEditChangeProperty(FPropertyChangedEvent& PropertyChangedEvent)
 {
-	if (UCustomizableSkeletalComponent* CustomizableSkeletalComponent = GetPrivate()->GetCustomizableSkeletalComponent())
+	Super::PostEditChangeProperty(PropertyChangedEvent);
+
+	if (PropertyChangedEvent.Property == nullptr)
 	{
-		return CustomizableSkeletalComponent->CustomizableObjectInstance;
+		return;
 	}
-	else
+
+	if (PropertyChangedEvent.Property->GetFName() == GET_MEMBER_NAME_CHECKED(UCustomizableObjectInstanceUsage, UsedCustomizableObjectInstance))
 	{
-		return UsedCustomizableObjectInstance;
+		UCustomizableObjectInstance* Instance = UsedCustomizableObjectInstance;
+		UsedCustomizableObjectInstance = OldCustomizableObjectInstance.Get(true);
+		SetCustomizableObjectInstance(Instance);
 	}
-}
-
-
-void UCustomizableObjectInstanceUsage::SetComponentName(const FName& Name)
-{
-	if (UCustomizableSkeletalComponent* CustomizableSkeletalComponent = GetPrivate()->GetCustomizableSkeletalComponent())
+	else if (PropertyChangedEvent.Property->GetFName() == GET_MEMBER_NAME_CHECKED(UCustomizableObjectInstanceUsage, UsedComponentName))
 	{
-		CustomizableSkeletalComponent->SetComponentName(Name);
+		SetComponentName(UsedComponentName);
 	}
-	else
+	else if (PropertyChangedEvent.Property->GetFName() == GET_MEMBER_NAME_CHECKED(UCustomizableObjectInstanceUsage, bUsedSkipSetReferenceSkeletalMesh))
+	{
+		SetSkipSetReferenceSkeletalMesh(bUsedSkipSetReferenceSkeletalMesh);
+	}
+	else if (PropertyChangedEvent.Property->GetFName() == GET_MEMBER_NAME_CHECKED(UCustomizableObjectInstanceUsage, bUsedSkipSetSkeletalMeshOnAttach))
 	{
-		UsedComponentIndex = INDEX_NONE;
-		UsedComponentName = Name;
+		SetSkipSetSkeletalMeshOnAttach(bUsedSkipSetSkeletalMeshOnAttach);
 	}
 }
+#endif
 
-
-FName UCustomizableObjectInstanceUsage::GetComponentName() const
+void UCustomizableObjectInstanceUsage::SetCustomizableObjectInstance(UCustomizableObjectInstance* InCustomizableObjectInstance)
 {
-	if (UCustomizableSkeletalComponent* CustomizableSkeletalComponent = GetPrivate()->GetCustomizableSkeletalComponent())
+	if (UsedCustomizableObjectInstance != InCustomizableObjectInstance)
 	{
-		return CustomizableSkeletalComponent->GetComponentName();
-	}
-	else
-	{
-		if (UsedComponentIndex == INDEX_NONE)
+		if (UsedCustomizableObjectInstance)
 		{
-			return UsedComponentName;	
+			UsedCustomizableObjectInstance->GetPrivate()->UnregisterCustomizableObjectInstanceUsage(*this);
 		}
-		else
+
+		UsedCustomizableObjectInstance = InCustomizableObjectInstance;
+
+		if (IsValid(UsedCustomizableObjectInstance))
 		{
-			return FName(FString::FromInt(UsedComponentIndex));
+			UsedCustomizableObjectInstance->GetPrivate()->RegisterCustomizableObjectInstanceUsage(*this);
 		}
+
+#if WITH_EDITOR
+		OldCustomizableObjectInstance = MakeWeakObjectPtr(UsedCustomizableObjectInstance);
+#endif
 	}
 }
 
 
-UCustomizableSkeletalComponent* UCustomizableObjectInstanceUsagePrivate::GetCustomizableSkeletalComponent() const
+UCustomizableObjectInstance* UCustomizableObjectInstanceUsage::GetCustomizableObjectInstance() const
 {
-	UObject* Outer = GetPublic()->GetOuter(); // UCustomizableSkeletalComponentPrivate
-	if (!Outer)
-	{
-		return nullptr;
-	}
+	return UsedCustomizableObjectInstance;
+}
+
 
-	return Cast<UCustomizableSkeletalComponent>(Outer->GetOuter());
+void UCustomizableObjectInstanceUsage::SetComponentName(const FName& Name)
+{
+	UsedComponentName = Name;
 }
 
+
+FName UCustomizableObjectInstanceUsage::GetComponentName() const
+{
+	return UsedComponentName;
+}
+
+
 UCustomizableObjectInstanceUsage* UCustomizableObjectInstanceUsagePrivate::GetPublic()
 {
 	UCustomizableObjectInstanceUsage* Public = StaticCast<UCustomizableObjectInstanceUsage*>(GetOuter());
@@ -150,53 +146,19 @@ const UCustomizableObjectInstanceUsage* UCustomizableObjectInstanceUsagePrivate:
 
 void UCustomizableObjectInstanceUsage::SetSkipSetReferenceSkeletalMesh(bool bSkip)
 {
-	if (UCustomizableSkeletalComponent* CustomizableSkeletalComponent = GetPrivate()->GetCustomizableSkeletalComponent())
-	{
-		CustomizableSkeletalComponent->SetSkipSetReferenceSkeletalMesh(bSkip);
-	}
-	else
-	{
-		bUsedSkipSetReferenceSkeletalMesh = bSkip;
-	}
+	bUsedSkipSetReferenceSkeletalMesh = bSkip;
 }
 
 
 bool UCustomizableObjectInstanceUsage::GetSkipSetReferenceSkeletalMesh() const
 {
-	if (UCustomizableSkeletalComponent* CustomizableSkeletalComponent = GetPrivate()->GetCustomizableSkeletalComponent())
-	{
-		return CustomizableSkeletalComponent->GetSkipSetReferenceSkeletalMesh();
-	}
-	else
-	{
-		return bUsedSkipSetReferenceSkeletalMesh;
-	}
+	return bUsedSkipSetReferenceSkeletalMesh;
 }
 
 
 void UCustomizableObjectInstanceUsage::SetSkipSetSkeletalMeshOnAttach(bool bSkip)
 {
-	if (UCustomizableSkeletalComponent* CustomizableSkeletalComponent = GetPrivate()->GetCustomizableSkeletalComponent())
-	{
-		CustomizableSkeletalComponent->SetSkipSetSkeletalMeshOnAttach(bSkip);
-	}
-	else
-	{
-		bUsedSkipSetSkeletalMeshOnAttach = bSkip;
-	}
-}
-
-
-bool UCustomizableObjectInstanceUsage::GetSkipSetSkeletalMeshOnAttach() const
-{
-	if (UCustomizableSkeletalComponent* CustomizableSkeletalComponent = GetPrivate()->GetCustomizableSkeletalComponent())
-	{
-		return CustomizableSkeletalComponent->GetSkipSetReferenceSkeletalMesh();
-	}
-	else
-	{
-		return bUsedSkipSetSkeletalMeshOnAttach;
-	}
+	bUsedSkipSetSkeletalMeshOnAttach = bSkip;
 }
 
 
@@ -213,7 +175,7 @@ void UCustomizableObjectInstanceUsage::AttachTo(USkeletalMeshComponent* Skeletal
 		GetPrivate()->bIsNetModeDedicatedServer = false;
 	}
 
-	if (!GetSkipSetSkeletalMeshOnAttach())
+	if (!bUsedSkipSetSkeletalMeshOnAttach)
 	{
 		GetPrivate()->bPendingSetReferenceSkeletalMesh = true;
 	}
@@ -222,11 +184,7 @@ void UCustomizableObjectInstanceUsage::AttachTo(USkeletalMeshComponent* Skeletal
 
 USkeletalMeshComponent* UCustomizableObjectInstanceUsage::GetAttachParent() const
 {
-	if (UCustomizableSkeletalComponent* CustomizableSkeletalComponent = GetPrivate()->GetCustomizableSkeletalComponent())
-	{
-		return Cast<USkeletalMeshComponent>(CustomizableSkeletalComponent->GetAttachParent());
-	}
-	else if(UsedSkeletalMeshComponent.IsValid())
+	if(UsedSkeletalMeshComponent.IsValid())
 	{
 		return UsedSkeletalMeshComponent.Get();
 	}
@@ -249,31 +207,27 @@ USkeletalMesh* UCustomizableObjectInstanceUsagePrivate::GetAttachedSkeletalMesh(
 
 	if (Parent)
 	{
-		return UE_MUTABLE_GETSKELETALMESHASSET(Parent);
+		return Parent->GetSkeletalMeshAsset();
 	}
 
 	return nullptr;
 }
 
 
-void UCustomizableObjectInstanceUsage::UpdateSkeletalMeshAsync(bool bNeverSkipUpdate, bool bIgnoreCloseDist, bool bForceHighPriority)
+void UCustomizableObjectInstanceUsage::UpdateSkeletalMeshAsync(bool bNeverSkipUpdate_DEPRECATED, bool bIgnoreCloseDist_DEPRECATED, bool bForceHighPriority)
 {
-	UCustomizableObjectInstance* CustomizableObjectInstance = GetCustomizableObjectInstance();
-
-	if (CustomizableObjectInstance)
+	if (UsedCustomizableObjectInstance)
 	{
-		CustomizableObjectInstance->UpdateSkeletalMeshAsync(bIgnoreCloseDist, bForceHighPriority);
+		UsedCustomizableObjectInstance->UpdateSkeletalMeshAsync(bIgnoreCloseDist_DEPRECATED, bForceHighPriority);
 	}
 }
 
 
-void UCustomizableObjectInstanceUsage::UpdateSkeletalMeshAsyncResult(FInstanceUpdateDelegate Callback, bool bIgnoreCloseDist, bool bForceHighPriority)
+void UCustomizableObjectInstanceUsage::UpdateSkeletalMeshAsyncResult(FInstanceUpdateDelegate Callback, bool bIgnoreCloseDist_DEPRECATED, bool bForceHighPriority)
 {
-	UCustomizableObjectInstance* CustomizableObjectInstance = GetCustomizableObjectInstance();
-
-	if (CustomizableObjectInstance)
+	if (UsedCustomizableObjectInstance)
 	{
-		CustomizableObjectInstance->UpdateSkeletalMeshAsyncResult(Callback, bIgnoreCloseDist, bForceHighPriority);
+		UsedCustomizableObjectInstance->UpdateSkeletalMeshAsyncResult(Callback, bIgnoreCloseDist_DEPRECATED, bForceHighPriority);
 	}
 }
 
diff --git a/Engine/Plugins/Mutable/Source/CustomizableObject/Private/MuCO/CustomizableObjectSystem.cpp b/Engine/Plugins/Mutable/Source/CustomizableObject/Private/MuCO/CustomizableObjectSystem.cpp
--- a/Engine/Plugins/Mutable/Source/CustomizableObject/Private/MuCO/CustomizableObjectSystem.cpp
+++ b/Engine/Plugins/Mutable/Source/CustomizableObject/Private/MuCO/CustomizableObjectSystem.cpp
@@ -38,6 +38,7 @@
 #include "Misc/ConfigCacheIni.h"
 #include "Engine/World.h"
 #include "MuT/UnrealPixelFormatOverride.h"
+#include "LevelEditorViewport.h"
 #else
 #include "Engine/Engine.h"
 #endif
@@ -385,38 +386,24 @@ bool TickWarnings(float DeltaTime)
 	
 	TSet<const UCustomizableObject*> Objects;
 	
-	for (UCustomizableObjectInstanceUsage* CustomizableObjectInstanceUsage : UCustomizableObjectSystemPrivate::GetCustomizableObjectInstanceUsages())
+	for (UCustomizableObjectInstance* Instance : UCustomizableObjectSystemPrivate::GetCustomizableObjectInstances())
 	{
-		const UCustomizableObjectInstance* Instance = CustomizableObjectInstanceUsage->GetCustomizableObjectInstance();
-		if (!Instance)
+		if (!IsValid(Instance))
 		{
 			continue;
 		}
 
-		const UCustomizableObject* Object = Cast<UCustomizableObject>(Instance->GetCustomizableObject());
-		if (!Object)
-		{
-			continue;
-		}
-		
-		const USkeletalMeshComponent* Parent = Cast<USkeletalMeshComponent>(CustomizableObjectInstanceUsage->GetAttachParent());
-		if (!Parent)
+		if (Instance->GetPrivate()->InstanceUsages.IsEmpty())
 		{
 			continue;
 		}
 
-		const UWorld* World = Parent->GetWorld();
-		if (!World)
+		const UCustomizableObject* Object = Cast<UCustomizableObject>(Instance->GetCustomizableObject());
+		if (!Object)
 		{
 			continue;
 		}
 
-		if (World->WorldType != EWorldType::PIE &&
-			World->WorldType != EWorldType::Game)
-		{
-			continue;
-		}
-		
 		if (Object->GetPrivate()->Status.Get() != FCustomizableObjectStatus::EState::ModelLoaded)
 		{
 			continue;
@@ -446,8 +433,8 @@ bool TickWarnings(float DeltaTime)
 		RemoveUnusedOnScreenMessages(KeysNotCompiled);
 	}
 
-	// Compiled without optimizations warning.
 #if WITH_EDITOR
+	// Compiled without optimizations warning.
 	for (const UCustomizableObject* Object : Objects)
 	{
 		TWeakObjectPtr<const UCustomizableObject> WeakObject = TWeakObjectPtr(Object);
@@ -465,11 +452,8 @@ bool TickWarnings(float DeltaTime)
 	}
 
 	RemoveUnusedOnScreenMessages(KeysNotOptimized);
-#endif
-	
 
 	// Is compilation out of date warning.
-#if WITH_EDITOR
 	if (!bOutOfDateAsync)
 	{
 		bOutOfDateAsync = true;
@@ -733,7 +717,7 @@ static FAutoConsoleVariableRef CVarGenerateInstancesWithinRange(
 
 TAutoConsoleVariable<int32> CVarMaxNumInstancesToDiscardPerTick(
 	TEXT("mutable.MaxNumInstancesToDiscardPerTick"),
-	30,
+	1,
 	TEXT("The maximum number of stale instances that will be discarded per tick by Mutable."),
 	ECVF_Scalability);
 
@@ -779,20 +763,22 @@ void FinishUpdateGlobal(const TSharedRef<FUpdateContextPrivate>& Context)
 	
 	if (IsValid(Instance))
 	{
-		UCustomizableInstancePrivate* PrivateInstance = Instance->GetPrivate();
-		
+		UCustomizableInstancePrivate* InstancePrivate = Instance->GetPrivate();
+
 		switch (Context->UpdateResult)
 		{
 		case EUpdateResult::Success:
 		case EUpdateResult::Warning:
-			PrivateInstance->SkeletalMeshStatus = ESkeletalMeshStatus::Success;
+			InstancePrivate->SkeletalMeshStatus = ESkeletalMeshStatus::Success;
 			
-			PrivateInstance->CommittedDescriptor = Context->MoveCommittedDescriptor();
-			PrivateInstance->CommittedDescriptorHash = Context->GetCapturedDescriptorHash();
+			InstancePrivate->CommittedDescriptor = Context->MoveCommittedDescriptor();
+			InstancePrivate->CommittedDescriptorHash = Context->GetCapturedDescriptorHash();
 			
 			// Delegates must be called only after updating the Instance flags.
 			Instance->UpdatedDelegate.Broadcast(Instance);
 			Instance->UpdatedNativeDelegate.Broadcast(Instance);
+
+			UCustomizableObjectSystemPrivate::RegisterInstanceToGeneratedList(*Instance);
 			break;
 
 		case EUpdateResult::ErrorOptimized:
@@ -807,14 +793,8 @@ void FinishUpdateGlobal(const TSharedRef<FUpdateContextPrivate>& Context)
 		case EUpdateResult::Error: 
 		case EUpdateResult::Error16BitBoneIndex:
 			{
-				UCustomizableInstancePrivate* InstancePrivate = Instance->GetPrivate();
-				check(InstancePrivate)
-				if (InstancePrivate)
-				{
-					InstancePrivate->InvalidateGeneratedData();
-				}
-				
-				PrivateInstance->SkeletalMeshStatus = ESkeletalMeshStatus::Error;
+				InstancePrivate->InvalidateGeneratedData();
+				InstancePrivate->SkeletalMeshStatus = ESkeletalMeshStatus::Error;
 			}
 			break;
 			
@@ -824,48 +804,31 @@ void FinishUpdateGlobal(const TSharedRef<FUpdateContextPrivate>& Context)
 		default:
 			unimplemented();
 		}
-	}
 
-	if (UCustomizableObjectSystem::IsUpdateResultValid(Context->UpdateResult))
-	{
-		MUTABLE_CPUPROFILER_SCOPE(FinishUpdateGlobal_Callbacks);
-		
-		// Call CustomizableObjectInstanceUsages updated callbacks.
-		if (Context->bOptimizedUpdate)
+
+		if (UCustomizableObjectSystem::IsUpdateResultValid(Context->UpdateResult))
 		{
-			for (TWeakObjectPtr<UCustomizableObjectInstanceUsage>& WeakUsage : Context->AttachedParentUpdated)
+			MUTABLE_CPUPROFILER_SCOPE(FinishUpdateGlobal_Callbacks);
+			TArrayView<UCustomizableObjectInstanceUsage*> InstanceUsagesView(InstancePrivate->InstanceUsages);
+
+			if (Context->bOptimizedUpdate)
 			{
-				TStrongObjectPtr<UCustomizableObjectInstanceUsage> InstanceUsage = WeakUsage.Pin();
-				if (!IsValid(InstanceUsage.Get()))
-				{
-					return;
-				}
+				InstanceUsagesView = TArrayView<UCustomizableObjectInstanceUsage*>(Context->AttachedParentUpdated);
+			}
 
-#if WITH_EDITOR
-				if (InstanceUsage->GetPrivate()->bIsNetModeDedicatedServer)
+			for (UCustomizableObjectInstanceUsage* InstanceUsage : InstanceUsagesView)
+			{
+				if (!IsValid(InstanceUsage))
 				{
 					continue;
 				}
-#endif
-				
-				InstanceUsage->GetPrivate()->Callbacks();
-			}
-		}
-		else
-		{
-			for (UCustomizableObjectInstanceUsage* InstanceUsage : UCustomizableObjectSystemPrivate::GetCustomizableObjectInstanceUsages())
-			{
 #if WITH_EDITOR
 				if (InstanceUsage->GetPrivate()->bIsNetModeDedicatedServer)
 				{
 					continue;
 				}
 #endif
-
-				if (InstanceUsage->GetCustomizableObjectInstance() == Instance)
-				{
-					InstanceUsage->GetPrivate()->Callbacks();
-				}
+				InstanceUsage->GetPrivate()->Callbacks();
 			}
 		}
 	}
@@ -1027,12 +990,17 @@ void UpdateSkeletalMesh(const TSharedRef<FUpdateContextPrivate>& Context)
 	check(IsInGameThread());
 
 	UCustomizableObjectInstance* CustomizableObjectInstance = Context->Instance.Get();
-	check(CustomizableObjectInstance);
+	UCustomizableObject* CustomizableObject = Context->Object.Get();
 
-	UCustomizableObject* CustomizableObject = CustomizableObjectInstance->GetCustomizableObject();
-	check(CustomizableObject);
+	if (!CustomizableObjectInstance || !CustomizableObject)
+	{
+		return;
+	}
+	
+	UCustomizableInstancePrivate* InstancePrivate = CustomizableObjectInstance->GetPrivate();
+	check(InstancePrivate != nullptr);
 
-	for (const TTuple<FName, TObjectPtr<USkeletalMesh>>& Pair : CustomizableObjectInstance->GetPrivate()->SkeletalMeshes)
+	for (const TTuple<FName, TObjectPtr<USkeletalMesh>>& Pair : InstancePrivate->SkeletalMeshes)
 	{
 		FPreSetSkeletalMeshParams Params;
 		Params.Instance = CustomizableObjectInstance;
@@ -1043,25 +1011,21 @@ void UpdateSkeletalMesh(const TSharedRef<FUpdateContextPrivate>& Context)
 	}
 	
 	uint32 ParentsWithOverrideMaterialsUpdatedCount = 0;
-	
-	UCustomizableInstancePrivate* InstancePrivate = CustomizableObjectInstance->GetPrivate();
-	check(InstancePrivate != nullptr);
 
-	for (UCustomizableObjectInstanceUsage* Usage : UCustomizableObjectSystemPrivate::GetCustomizableObjectInstanceUsages())
+	for (UCustomizableObjectInstanceUsage* InstanceUsage : InstancePrivate->InstanceUsages)
 	{
-#if WITH_EDITOR
-		if (Usage->GetPrivate()->bIsNetModeDedicatedServer)
+		if (!IsValid(InstanceUsage))
 		{
 			continue;
 		}
-#endif
 
-		if (Usage->GetCustomizableObjectInstance() != CustomizableObjectInstance)
+#if WITH_EDITOR
+		if (InstanceUsage->GetPrivate()->bIsNetModeDedicatedServer)
 		{
 			continue;
 		}
-
-		USkeletalMeshComponent* Parent = Cast<USkeletalMeshComponent>(Usage->GetAttachParent());
+#endif
+		USkeletalMeshComponent* Parent = Cast<USkeletalMeshComponent>(InstanceUsage->GetAttachParent());
 		if (!Parent)
 		{
 			continue;
@@ -1071,9 +1035,9 @@ void UpdateSkeletalMesh(const TSharedRef<FUpdateContextPrivate>& Context)
 
 		bool bAttachedParentUpdated = false;
 
-		Usage->GetPrivate()->bPendingSetReferenceSkeletalMesh = false;
+		InstanceUsage->GetPrivate()->bPendingSetReferenceSkeletalMesh = false;
 
-		USkeletalMesh* SkeletalMesh = CustomizableObjectInstance->GetComponentMeshSkeletalMesh(Usage->GetComponentName());
+		USkeletalMesh* SkeletalMesh = CustomizableObjectInstance->GetComponentMeshSkeletalMesh(InstanceUsage->GetComponentName());
 		if (SkeletalMesh != Parent->GetSkeletalMeshAsset())
 		{
 			Parent->SetSkeletalMesh(SkeletalMesh, RequiresReinitPose(Parent->GetSkeletalMeshAsset(), SkeletalMesh));
@@ -1088,7 +1052,7 @@ void UpdateSkeletalMesh(const TSharedRef<FUpdateContextPrivate>& Context)
 
 		TArray<TObjectPtr<UMaterialInterface>> NewOverridenMaterials;
 
-		FCustomizableInstanceComponentData* ComponentData = CustomizableObjectInstance->GetPrivate()->GetComponentData(Usage->GetComponentName());
+		FCustomizableInstanceComponentData* ComponentData = CustomizableObjectInstance->GetPrivate()->GetComponentData(InstanceUsage->GetComponentName());
 		
 		const bool bIsTransientMesh = SkeletalMesh ? SkeletalMesh->HasAllFlags(EObjectFlags::RF_Transient) : false;
 		const bool bUseOverrideMaterials = !bIsTransientMesh || Context->bUseMeshCache;
@@ -1151,12 +1115,12 @@ void UpdateSkeletalMesh(const TSharedRef<FUpdateContextPrivate>& Context)
 
 		if (bAttachedParentUpdated)
 		{
-			Context->AttachedParentUpdated.Add(Usage);
+			Context->AttachedParentUpdated.Add(InstanceUsage);
 		}
 		
 		for (const UCustomizableObjectExtension* Extension : ICustomizableObjectModule::Get().GetRegisteredExtensions())
 		{
-			Extension->OnCustomizableObjectInstanceUsageUpdated(*Usage, InstancePrivate->ExtensionData);
+			Extension->OnCustomizableObjectInstanceUsageUpdated(*InstanceUsage, InstancePrivate->ExtensionData);
 		}
 	}
 
@@ -1167,9 +1131,134 @@ void UpdateSkeletalMesh(const TSharedRef<FUpdateContextPrivate>& Context)
 }
 
 
-bool UCustomizableObjectSystemPrivate::IsDiscardedByDistance(const UCustomizableInstancePrivate* InstancePrivate) const
+void UCustomizableObjectSystemPrivate::UpdateViewLocations()
 {
-	return bGenerateInstancesWithinRange && InstancePrivate->MinSquareDistFromComponentToPlayer > GenerationRangeSquare;
+	if (!bGenerateInstancesWithinRange)
+	{
+		return;
+	}
+
+	MUTABLE_CPUPROFILER_SCOPE(UpdateViewLocations);
+	ViewLocations.Reset(ViewCenters.Num());
+	const int32 NumViewCenters = ViewCenters.Num();
+
+	for (const TWeakObjectPtr<const AActor> ViewCenter : ViewCenters)
+	{
+		if (ViewCenter.IsValid())
+		{
+			ViewLocations.Add(ViewCenter->GetActorLocation());
+		}
+	}
+}
+
+
+bool UCustomizableObjectSystemPrivate::IsDiscardedByDistance(const UCustomizableObjectInstance& Instance) const
+{
+	if (!bGenerateInstancesWithinRange)
+	{
+		return false;
+	}
+
+#if !WITH_EDITOR
+	if (ViewLocations.IsEmpty())
+	{
+		return false;
+	}
+#endif
+
+	MUTABLE_CPUPROFILER_SCOPE(IsDiscardedByDistance);
+
+	float MinSquareDistFromComponentToPlayer = FLT_MAX;
+	for (UCustomizableObjectInstanceUsage* InstanceUsage : Instance.GetPrivate()->InstanceUsages)
+	{
+		if (!IsValid(InstanceUsage))
+		{
+			continue;
+		}
+
+		USkeletalMeshComponent* Parent = InstanceUsage->GetAttachParent();
+		if (!Parent)
+		{
+			continue;
+		}
+		const AActor* ParentActor = Parent->GetAttachmentRootActor();
+		if (!IsValid(ParentActor))
+		{
+			continue;
+		}
+		const APawn* Pawn = Cast<const APawn>(ParentActor);
+		if (Pawn && Pawn->IsPlayerControlled())
+		{
+			return false;
+		}
+		const FVector ActorLocation = ParentActor->GetActorLocation();
+		for (const FVector& ViewLocation : ViewLocations)
+		{
+			MinSquareDistFromComponentToPlayer = FMath::Min(FVector::DistSquared(ViewLocation, ActorLocation), MinSquareDistFromComponentToPlayer);
+		}
+#if WITH_EDITOR
+		UWorld* LocalWorld = Parent->GetWorld();
+		if (LocalWorld && LocalWorld->WorldType == EWorldType::Editor)
+		{
+			for (FLevelEditorViewportClient* LevelVC : GEditor->GetLevelViewportClients())
+			{
+				if (LevelVC && LevelVC->IsPerspective() && LevelVC->GetWorld() == LocalWorld)
+				{
+					MinSquareDistFromComponentToPlayer = FMath::Min(FVector::DistSquared(LevelVC->GetViewLocation(), ActorLocation), MinSquareDistFromComponentToPlayer);
+					break;
+				}
+			}
+		}
+#endif
+		if (MinSquareDistFromComponentToPlayer <= GenerationRangeSquare)
+		{
+			return false;
+		}
+	}
+	return true;
+}
+
+
+void UCustomizableObjectSystemPrivate::DiscardInstance(UCustomizableObjectInstance& Instance)
+{
+	MUTABLE_CPUPROFILER_SCOPE(DiscardInstance);
+
+	UCustomizableInstancePrivate* InstancePrivate = Instance.GetPrivate();
+	InstancePrivate->InvalidateGeneratedData();
+
+	UCustomizableObject* Object = Instance.GetCustomizableObject();
+	check(Object);
+
+	for (UCustomizableObjectInstanceUsage* InstanceUsage : InstancePrivate->InstanceUsages)
+	{
+		if (!IsValid(InstanceUsage))
+		{
+			continue;
+		}
+
+		USkeletalMeshComponent* Parent = InstanceUsage->GetAttachParent();
+		if (!Parent)
+		{
+			continue;
+		}
+
+		Parent->EmptyOverrideMaterials();
+		USkeletalMesh* SkeletalMesh = nullptr;
+
+		if (!InstanceUsage->GetSkipSetReferenceSkeletalMesh() &&
+			Object->bEnableUseRefSkeletalMeshAsPlaceholder)
+		{
+			const FName& ComponentName = InstanceUsage->GetComponentName();
+			SkeletalMesh = Object->GetComponentMeshReferenceSkeletalMesh(ComponentName);
+		}
+
+		Parent->SetSkeletalMesh(SkeletalMesh);
+
+		for (const UCustomizableObjectExtension* Extension : ICustomizableObjectModule::Get().GetRegisteredExtensions())
+		{
+			Extension->OnCustomizableObjectInstanceUsageDiscarded(*InstanceUsage);
+		}
+	}
 }
 
 
@@ -1178,23 +1267,24 @@ void UCustomizableObjectSystemPrivate::EnqueueUpdateSkeletalMesh(const TSharedRe
 	MUTABLE_CPUPROFILER_SCOPE(FCustomizableObjectSystemPrivate::EnqueueUpdateSkeletalMesh);
 	check(IsInGameThread());
 	
-	UCustomizableObject* Object = Context->Object.Get();
 	UCustomizableObjectInstance* Instance = Context->Instance.Get();
+	UCustomizableObject* Object = Context->Object.Get();
 
-	const FName ObjectName = Object ? Object->GetFName() : NAME_None;
-	const FName InstanceName = Instance ? Instance->GetFName() : NAME_None;
-	UE_LOGF(LogMutable, Verbose, "Enqueue Update Skeletal Mesh Async. CustomizableObject=%ls Instance=%ls, Frame=%d", *ObjectName.ToString(), *InstanceName.ToString(), GFrameNumber);
-	
-	if (!Context->IsContextValid()) 
+	if (!Instance || !Object)
 	{
 		Context->UpdateResult = EUpdateResult::Error;
 		FinishUpdateGlobal(Context);
 		return;
 	}
 
-	check(Object);
-	check(Instance);
+	if (!Context->IsContextValid()) 
+	{
+		Context->UpdateResult = EUpdateResult::Error;
+		FinishUpdateGlobal(Context);
+		return;
+	}
 	
+	UE_LOGF(LogMutable, Verbose, "Enqueue Update Skeletal Mesh Async. CustomizableObject=%ls Instance=%ls, Frame=%d", *Object->GetFName().ToString(), *Instance->GetFName().ToString(), GFrameNumber);
 	UCustomizableInstancePrivate* InstancePrivate = Instance->GetPrivate();
 		
 	if (!bIsMutableEnabled)
@@ -1230,13 +1320,6 @@ void UCustomizableObjectSystemPrivate::EnqueueUpdateSkeletalMesh(const TSharedRe
 		FinishUpdateGlobal(Context);
 		return;
 	}
-	
-	if (IsDiscardedByDistance(InstancePrivate))
-	{
-		Context->UpdateResult = EUpdateResult::ErrorDiscarded;
-		FinishUpdateGlobal(Context);
-		return;
-	}
 
 	if (!Context->bForce)
 	{
@@ -1707,8 +1790,7 @@ namespace impl
 	{
 		MUTABLE_CPUPROFILER_SCOPE(Task_Mutable_GetSurfaces)
 		
-		const UCustomizableObjectInstance* Instance = OperationData->Instance.Get();
-		const UCustomizableObject* CustomizableObject = Instance->GetCustomizableObject();
+		const UCustomizableObject* CustomizableObject = OperationData->Object.Get();
 		const UModelResources& ModelResources = *CustomizableObject->GetPrivate()->GetModelResources();
 
 		const UE::Mutable::Private::FInstance* MutableInstance = OperationData->MutableInstance.Get();
@@ -1782,8 +1864,7 @@ namespace impl
 	{
 		MUTABLE_CPUPROFILER_SCOPE(Task_Mutable_GetMaterials)
 
-		const UCustomizableObjectInstance* Instance = Context->Instance.Get();
-		const UCustomizableObject* CustomizableObject = Instance->GetCustomizableObject();
+		const UCustomizableObject* CustomizableObject = Context->Object.Get();
 		const UModelResources* ModelResources = CustomizableObject->GetPrivate()->GetModelResources();
 
 		UE::Tasks::FTask LastTask = UE::Tasks::MakeCompletedTask<void>();
@@ -2820,96 +2901,19 @@ namespace impl
 
 void UCustomizableObjectSystemPrivate::UpdateTick(const bool bBlocking)
 {
-	{
-		MUTABLE_CPUPROFILER_SCOPE(UpdateTick_Usage);
-
-		for (UCustomizableObjectInstanceUsage* CustomizableObjectInstanceUsage : CustomizableObjectInstanceUsages)
-		{
-			UCustomizableObjectInstance* Instance = CustomizableObjectInstanceUsage->GetCustomizableObjectInstance();
-			if (!Instance)
-			{
-				continue;
-			}
-
-			UCustomizableObject* Object = Instance->GetCustomizableObject();
-			if (!Object)
-			{
-				continue;
-			}
-
-			USkeletalMeshComponent* Parent = Cast<USkeletalMeshComponent>(CustomizableObjectInstanceUsage->GetAttachParent());
-			if (!Parent)
-			{
-				continue;
-			}
-
-			Instance->GetPrivate()->bAllowAutomaticUpdate = true;
-
-			UCustomizableObjectInstanceUsagePrivate* UsagePrivate = CustomizableObjectInstanceUsage->GetPrivate();
-
-			if (Instance->GetPrivate()->HasCOInstanceFlags(ECOInstanceFlags::Discarded))
-			{
-				Instance->GetPrivate()->InvalidateGeneratedData();
-
-				Parent->EmptyOverrideMaterials();
-
-				USkeletalMesh* SkeletalMesh = nullptr;
-
-				if (!CustomizableObjectInstanceUsage->GetSkipSetReferenceSkeletalMesh() &&
-					Object->bEnableUseRefSkeletalMeshAsPlaceholder)
-				{
-					const FName& ComponentName = CustomizableObjectInstanceUsage->GetComponentName();
-					SkeletalMesh = Object->GetComponentMeshReferenceSkeletalMesh(ComponentName);
-				}
-
-				Parent->SetSkeletalMesh(SkeletalMesh);
-
-				for (const UCustomizableObjectExtension* Extension : ICustomizableObjectModule::Get().GetRegisteredExtensions())
-				{
-					Extension->OnCustomizableObjectInstanceUsageDiscarded(*CustomizableObjectInstanceUsage);
-				}
-
-				UsagePrivate->bPendingSetReferenceSkeletalMesh = false;
-			}
-
-			else if (UsagePrivate->bPendingSetReferenceSkeletalMesh)
-			{
-				if (!CustomizableObjectInstanceUsage->GetSkipSetReferenceSkeletalMesh() &&
-					Object->bEnableUseRefSkeletalMeshAsPlaceholder)
-				{
-					const FName& ComponentName = CustomizableObjectInstanceUsage->GetComponentName();
+	TickPendingSetReferenceSkeletalMesh();
 
-					if (USkeletalMesh* ReferenceSkeletalMesh = Object->GetComponentMeshReferenceSkeletalMesh(ComponentName))
-					{
-						Parent->EmptyOverrideMaterials();
-						Parent->SetSkeletalMesh(ReferenceSkeletalMesh);
-					}
-				}
-
-				UsagePrivate->bPendingSetReferenceSkeletalMesh = false;
-			}
-		}
-	}
-	
 	bool bPendingCompilation = false;
 #if WITH_EDITOR
 	ICustomizableObjectEditorModule* EditorModule = ICustomizableObjectEditorModule::Get();
 	bPendingCompilation = EditorModule && EditorModule->GetNumCompileRequests() > 0;
 #endif
+	
+	TickDiscardInstances();
 
 	// Get a new operation if we aren't working on one
 	if (!CurrentMutableOperation && bIsMutableEnabled && !bPendingCompilation)
 	{
-		if (bGenerateInstancesWithinRange)
-		{
-			UpdateInstanceDistances();
-
-			if (ICustomizableObjectEditorModule* Module = ICustomizableObjectEditorModule::Get())
-			{
-				Module->UpdateInstanceDistanceToLevelEditorCamera();
-			}
-		}
-
 		bool bUpdatingInstance = false;
 		{
 			MUTABLE_CPUPROFILER_SCOPE(UpdateTick_Queue);
@@ -2945,33 +2949,9 @@ void UCustomizableObjectSystemPrivate::UpdateTick(const bool bBlocking)
 			}
 		}
 
-		int32 MaxNumInstancesToDiscard = CVarMaxNumInstancesToDiscardPerTick.GetValueOnGameThread();
-
+		if (!bUpdatingInstance)
 		{
-			MUTABLE_CPUPROFILER_SCOPE(UpdateTick_AutomaticUpdates);
-			for (UCustomizableObjectInstance* CustomizableObjectInstance : CustomizableObjectInstances)
-			{
-				UCustomizableInstancePrivate* CustomizableObjectInstancePrivate = CustomizableObjectInstance->GetPrivate();
-				CustomizableObjectInstancePrivate->ClearCOInstanceFlags(ECOInstanceFlags::Discarded);
-
-				if (IsDiscardedByDistance(CustomizableObjectInstancePrivate))
-				{
-					if (MaxNumInstancesToDiscard > 0 &&
-						CustomizableObjectInstancePrivate->SkeletalMeshStatus == ESkeletalMeshStatus::Success)
-					{
-						CustomizableObjectInstancePrivate->SetCOInstanceFlags(ECOInstanceFlags::Discarded);
-						--MaxNumInstancesToDiscard;
-					}
-				}
-				else if (!bUpdatingInstance &&
-					(CustomizableObjectInstancePrivate->bAllowAutomaticUpdate && CustomizableObjectInstancePrivate->SkeletalMeshStatus == ESkeletalMeshStatus::NotGenerated))
-				{
-					bUpdatingInstance = TryStartAutomaticUpdateForInstance(*CustomizableObjectInstance);
-				}
-
-				CustomizableObjectInstancePrivate->MinSquareDistFromComponentToPlayer = FLT_MAX;
-				CustomizableObjectInstancePrivate->bAllowAutomaticUpdate = false;
-			}
+			TickAutomaticUpdateInstances();
 		}
 
 		// Update the streaming limit if it has changed. It is safe to do this now.
@@ -3078,43 +3058,234 @@ bool UCustomizableObjectSystem::IsUpdating(const UCustomizableObjectInstance* In
 	return GetPrivate()->IsUpdating(*Instance);
 }
 
-TArray<UCustomizableObjectInstanceUsage*> UCustomizableObjectSystemPrivate::CustomizableObjectInstanceUsages;
 TArray<UCustomizableObjectInstance*> UCustomizableObjectSystemPrivate::CustomizableObjectInstances;
 
-void UCustomizableObjectSystemPrivate::RegisterCustomizableObjectInstanceUsage(UCustomizableObjectInstanceUsage& Obj)
+void UCustomizableObjectSystemPrivate::RegisterCustomizableObjectInstance(UCustomizableObjectInstance& Obj)
 {
 	if (!Obj.IsTemplate())
 	{
-		check(!UCustomizableObjectSystemPrivate::CustomizableObjectInstanceUsages.Contains(&Obj));
-		UCustomizableObjectSystemPrivate::CustomizableObjectInstanceUsages.Add(&Obj);
+		check(!UCustomizableObjectSystemPrivate::CustomizableObjectInstances.Contains(&Obj));
+		UCustomizableObjectSystemPrivate::CustomizableObjectInstances.Add(&Obj);
 	}
 }
 
-void UCustomizableObjectSystemPrivate::UnregisterCustomizableObjectInstanceUsage(UCustomizableObjectInstanceUsage& Obj)
+
+void UCustomizableObjectSystemPrivate::UnregisterCustomizableObjectInstance(UCustomizableObjectInstance& Obj)
 {
 	if (!Obj.IsTemplate())
 	{
-		UCustomizableObjectSystemPrivate::CustomizableObjectInstanceUsages.Remove(&Obj);
+		UCustomizableObjectSystemPrivate::CustomizableObjectInstances.RemoveSingleSwap(&Obj);
+		UCustomizableObjectSystemPrivate::PendingSetReferenceSkeletalMesh.RemoveSingleSwap(&Obj);
+		UCustomizableObjectSystemPrivate::GeneratedInstances.RemoveSingleSwap(&Obj);
+		UCustomizableObjectSystemPrivate::AutomaticUpdateInstances.RemoveSingleSwap(&Obj);
 	}
 }
 
-void UCustomizableObjectSystemPrivate::RegisterCustomizableObjectInstance(UCustomizableObjectInstance& Obj)
+TArray<UCustomizableObjectInstance*> UCustomizableObjectSystemPrivate::PendingSetReferenceSkeletalMesh;
+TArray<UCustomizableObjectInstance*> UCustomizableObjectSystemPrivate::GeneratedInstances;
+TArray<UCustomizableObjectInstance*> UCustomizableObjectSystemPrivate::AutomaticUpdateInstances;
+
+void UCustomizableObjectSystemPrivate::RegisterInstanceToPendinSetReferenceSkeletalMeshList(UCustomizableObjectInstance& Obj)
 {
 	if (!Obj.IsTemplate())
 	{
-		check(!UCustomizableObjectSystemPrivate::CustomizableObjectInstances.Contains(&Obj));
-		UCustomizableObjectSystemPrivate::CustomizableObjectInstances.Add(&Obj);
+		UCustomizableObjectSystemPrivate::PendingSetReferenceSkeletalMesh.AddUnique(&Obj);
 	}
 }
 
-void UCustomizableObjectSystemPrivate::UnregisterCustomizableObjectInstance(UCustomizableObjectInstance& Obj)
+
+void UCustomizableObjectSystemPrivate::RegisterInstanceToGeneratedList(UCustomizableObjectInstance& Obj)
 {
 	if (!Obj.IsTemplate())
 	{
-		UCustomizableObjectSystemPrivate::CustomizableObjectInstances.Remove(&Obj);
+		UCustomizableObjectSystemPrivate::GeneratedInstances.AddUnique(&Obj);
 	}
 }
 
+
+void UCustomizableObjectSystemPrivate::UnregisterInstancefromGeneratedList(UCustomizableObjectInstance& Obj)
+{
+	if (!Obj.IsTemplate())
+	{
+		UCustomizableObjectSystemPrivate::GeneratedInstances.RemoveSingleSwap(&Obj);
+	}
+}
+
+
+void UCustomizableObjectSystemPrivate::RegisterInstanceToAutomaticUpdateList(UCustomizableObjectInstance& Obj)
+{
+	if (!Obj.IsTemplate())
+	{
+		UCustomizableObjectSystemPrivate::AutomaticUpdateInstances.AddUnique(&Obj);
+	}
+}
+
+
+void UCustomizableObjectSystemPrivate::UnregisterInstanceFromAutomaticUpdateList(UCustomizableObjectInstance& Obj)
+{
+	if (!Obj.IsTemplate())
+	{
+		UCustomizableObjectSystemPrivate::AutomaticUpdateInstances.RemoveSingleSwap(&Obj);
+	}
+}
+
+
+void UCustomizableObjectSystemPrivate::TickPendingSetReferenceSkeletalMesh()
+{
+	MUTABLE_CPUPROFILER_SCOPE(TickPendingSetReferenceSkeletalMesh);
+
+	for (int32 Index = PendingSetReferenceSkeletalMesh.Num() - 1; Index >= 0; --Index)
+	{
+		UCustomizableObjectInstance* Instance = PendingSetReferenceSkeletalMesh.Pop(EAllowShrinking::No);
+		if (!IsValid(Instance))
+		{
+			continue;
+		}
+
+		UCustomizableInstancePrivate* InstancePrivate = Instance->GetPrivate();
+
+		UCustomizableObject* Object = Instance->GetCustomizableObject();
+		if (!Object || !Object->bEnableUseRefSkeletalMeshAsPlaceholder)
+		{
+			continue;
+		}
+
+		for (UCustomizableObjectInstanceUsage* InstanceUsage : InstancePrivate->InstanceUsages)
+		{
+			if (!IsValid(InstanceUsage))
+			{
+				continue;
+			}
+
+			UCustomizableObjectInstanceUsagePrivate* InstanceUsagePrivate = InstanceUsage->GetPrivate();
+
+			if (!InstanceUsagePrivate->bPendingSetReferenceSkeletalMesh)
+			{
+				continue;
+			}
+
+			if (InstanceUsage->GetSkipSetReferenceSkeletalMesh())
+			{
+				continue;
+			}
+
+			USkeletalMeshComponent* Parent = Cast<USkeletalMeshComponent>(InstanceUsage->GetAttachParent());
+			if (!Parent)
+			{
+				continue;
+			}
+
+			const FName& ComponentName = InstanceUsage->GetComponentName();
+			if (USkeletalMesh* ReferenceSkeletalMesh = Object->GetComponentMeshReferenceSkeletalMesh(ComponentName))
+			{
+				Parent->EmptyOverrideMaterials();
+				Parent->SetSkeletalMesh(ReferenceSkeletalMesh);
+			}
+
+			InstanceUsagePrivate->bPendingSetReferenceSkeletalMesh = false;
+		}
+	}
+
+	PendingSetReferenceSkeletalMesh.Shrink();
+}
+
+
+void UCustomizableObjectSystemPrivate::TickDiscardInstances()
+{
+	if (!bGenerateInstancesWithinRange)
+	{
+		return;
+	}
+
+	MUTABLE_CPUPROFILER_SCOPE(TickDiscardInstances);
+
+	UpdateViewLocations();
+
+	int32 MaxNumInstancesToDiscard = CVarMaxNumInstancesToDiscardPerTick.GetValueOnGameThread();
+
+	constexpr int32 NumInstancesPerTick = 25;
+	const int32 NumInstances = GeneratedInstances.Num();
+	int32 NumRemainingInstances = FMath::Min(NumInstances, NumInstancesPerTick);
+
+	while (NumRemainingInstances > 0)
+	{
+		--NumRemainingInstances;
+
+		DiscardIndex = DiscardIndex < GeneratedInstances.Num() ? DiscardIndex : 0;
+
+		UCustomizableObjectInstance* Instance = GeneratedInstances[DiscardIndex];
+		++DiscardIndex;
+
+		if (!IsValid(Instance))
+		{
+			continue;
+		}
+
+		if (IsDiscardedByDistance(*Instance))
+		{
+			DiscardInstance(*Instance);
+
+			GeneratedInstances.RemoveSingleSwap(Instance, EAllowShrinking::No);
+			--DiscardIndex;
+
+			AutomaticUpdateInstances.AddUnique(Instance);
+
+			if (--MaxNumInstancesToDiscard == 0)
+			{
+				break;
+			}
+		}
+	}
+
+	if (GeneratedInstances.Num() < (GeneratedInstances.GetAllocatedSize() / 2) / GeneratedInstances.GetTypeSize())
+	{
+		GeneratedInstances.Shrink();
+	}
+}
+
+
+void UCustomizableObjectSystemPrivate::TickAutomaticUpdateInstances()
+{
+	MUTABLE_CPUPROFILER_SCOPE(TickAutomaticUpdateInstances);
+
+	constexpr int32 NumInstancesPerTick = 50;
+	
+	const int32 NumInstances = AutomaticUpdateInstances.Num();
+	int32 NumRemainingInstances = FMath::Min(NumInstances, NumInstancesPerTick);
+	
+	while (NumRemainingInstances > 0)
+	{
+		--NumRemainingInstances;
+		AutomaticUpdateIndex = AutomaticUpdateIndex < AutomaticUpdateInstances.Num() ? AutomaticUpdateIndex : 0;
+
+		UCustomizableObjectInstance* Instance = AutomaticUpdateInstances[AutomaticUpdateIndex];
+		++AutomaticUpdateIndex;
+
+		if (!IsValid(Instance))
+		{
+			continue;
+		}
+
+		UCustomizableInstancePrivate* InstancePrivate = Instance->GetPrivate();
+
+		const bool bNeedsAutomaticUpdate = InstancePrivate->SkeletalMeshStatus != ESkeletalMeshStatus::Error &&
+			!IsDiscardedByDistance(*Instance);
+
+		if (bNeedsAutomaticUpdate && TryStartAutomaticUpdateForInstance(*Instance))
+		{
+			AutomaticUpdateInstances.RemoveSingleSwap(Instance, EAllowShrinking::No);
+			--AutomaticUpdateIndex;
+			break;
+		}
+	}
+
+	if (AutomaticUpdateInstances.Num() < (AutomaticUpdateInstances.GetAllocatedSize() / 2) / AutomaticUpdateInstances.GetTypeSize())
+	{
+		AutomaticUpdateInstances.Shrink();
+	}
+}
+
+
 bool UCustomizableObjectSystemPrivate::IsUsingBenchmarkingSettings()
 {
 	return bUseBenchmarkingSettings;
@@ -3215,100 +3386,6 @@ void UCustomizableObjectSystem::ClearViewCenters()
 }
 
 
-void UCustomizableObjectSystemPrivate::UpdateInstanceDistances()
-{
-	MUTABLE_CPUPROFILER_SCOPE(UpdateInstanceDistances);
-
-	if (CustomizableObjectInstanceUsages.IsEmpty())
-	{
-		return;
-	}
-
-	UCustomizableObjectSystem* CustomizableObjectSystem = UCustomizableObjectSystem::GetInstance();
-
-	TArray<const FVector> ViewLocations;
-	ViewLocations.Reset(ViewCenters.Num());
-
-	const int32 NumViewCenters = ViewCenters.Num();
-	for (const TWeakObjectPtr<const AActor> ViewCenter : ViewCenters)
-	{
-		if (ViewCenter.IsValid())
-		{
-			ViewLocations.Add(ViewCenter->GetActorLocation());
-		}
-	}
-
-	if (ViewLocations.IsEmpty())
-	{
-		for (UCustomizableObjectInstanceUsage* CustomizableObjectInstanceUsage : CustomizableObjectInstanceUsages)
-		{
-			USkeletalMeshComponent* Parent = Cast<USkeletalMeshComponent>(CustomizableObjectInstanceUsage->GetAttachParent());
-			if (!Parent || !Parent->GetWorld())
-			{
-				continue;
-			}
-
-			APlayerController* FirstLocalPlayerController = Parent->GetWorld()->GetFirstPlayerController();
-			if (FirstLocalPlayerController && FirstLocalPlayerController->GetPawn())
-			{
-				ViewLocations.Add(FirstLocalPlayerController->GetPawn()->GetActorLocation());
-				break;
-			}
-		}
-	}
-
-	if (ViewLocations.Num())
-	{
-		for (UCustomizableObjectInstanceUsage* CustomizableObjectInstanceUsage : CustomizableObjectInstanceUsages)
-		{
-#if WITH_EDITOR
-			if (CustomizableObjectInstanceUsage->GetPrivate()->bIsNetModeDedicatedServer)
-			{
-				continue;
-			}
-
-#endif // WITH_EDITOR
-
-			UCustomizableObjectInstance* CustomizableObjectInstance = CustomizableObjectInstanceUsage->GetCustomizableObjectInstance();
-			if (!CustomizableObjectInstance)
-			{
-				continue;
-			}
-
-			const USkeletalMeshComponent* SkeletalMeshComponent = CustomizableObjectInstanceUsage->GetAttachParent();
-			if (!SkeletalMeshComponent)
-			{
-				continue;
-			}
-
-			const AActor* ParentActor = SkeletalMeshComponent->GetAttachmentRootActor();
-			if (!IsValid(ParentActor))
-			{
-				continue;
-			}
-
-			UCustomizableInstancePrivate* InstancePrivate = CustomizableObjectInstance->GetPrivate();
-
-			const APawn* Pawn = Cast<const APawn>(ParentActor);
-			if (Pawn && Pawn->IsPlayerControlled())
-			{
-				InstancePrivate->SetCOInstanceFlags(ECOInstanceFlags::UsedByPlayerOrNearIt);
-				InstancePrivate->MinSquareDistFromComponentToPlayer = -1.0;
-				continue;
-			}
-
-			InstancePrivate->ClearCOInstanceFlags(ECOInstanceFlags::UsedByPlayerOrNearIt);
-
-			const FVector ActorLocation = ParentActor->GetActorLocation();
-			for (const FVector& ViewLocation : ViewLocations)
-			{
-				InstancePrivate->MinSquareDistFromComponentToPlayer = FMath::Min(FVector::DistSquared(ViewLocation, ActorLocation), InstancePrivate->MinSquareDistFromComponentToPlayer);
-			}
-		}
-	}
-}
-
-
 void UCustomizableObjectSystem::AddUncompiledCOWarning(const UCustomizableObject& InObject, FString const* OptionalLogInfo)
 {
 	FString Msg;
@@ -3434,7 +3511,7 @@ bool UCustomizableObjectSystemPrivate::TryStartUpdate(const TSharedRef<FUpdateCo
 	InstancePrivate->ClearCOInstanceFlags(ECOInstanceFlags::HasPendingHighPriorityUpdate | ECOInstanceFlags::HasPendingLowPriorityUpdate);
 
 	// Only update resources if the instance is in range (it could have got far from the player since the task was queued)
-	if (IsDiscardedByDistance(InstancePrivate))
+	if (IsDiscardedByDistance(*Instance))
 	{
 		Context->UpdateResult = EUpdateResult::ErrorDiscarded;
 		FinishUpdateGlobal(Context);
@@ -3496,6 +3573,19 @@ bool UCustomizableObjectSystemPrivate::TryStartAutomaticUpdateForInstance(UCusto
 		return false;
 	}
 
+	if (Context->GetCapturedDescriptorHash() == InstancePrivate->CommittedDescriptorHash)
+	{
+		Context->bOptimizedUpdate = true;
+
+		// The user may have changed the AttachParent and we need to re-customize it.
+		// In case nothing need to be re-customized, the update will be considered ErrorOptimized.
+		UpdateSkeletalMesh(Context);
+		Context->UpdateResult = Context->AttachedParentUpdated.IsEmpty() ? EUpdateResult::ErrorOptimized : EUpdateResult::Success;
+
+		FinishUpdateGlobal(Context);
+		return true;
+	}
+
 	Context->Parameters = Context->GetCapturedDescriptor().GetParameters();
 
 	StartUpdate(Context);
diff --git a/Engine/Plugins/Mutable/Source/CustomizableObject/Private/MuCO/CustomizableSkeletalComponent.cpp b/Engine/Plugins/Mutable/Source/CustomizableObject/Private/MuCO/CustomizableSkeletalComponent.cpp
--- a/Engine/Plugins/Mutable/Source/CustomizableObject/Private/MuCO/CustomizableSkeletalComponent.cpp
+++ b/Engine/Plugins/Mutable/Source/CustomizableObject/Private/MuCO/CustomizableSkeletalComponent.cpp
@@ -18,51 +18,66 @@ UCustomizableSkeletalComponent::UCustomizableSkeletalComponent()
 	Private = CreateDefaultSubobject<UCustomizableSkeletalComponentPrivate>(TEXT("Private"));
 }
 
-
-UCustomizableSkeletalComponentPrivate::UCustomizableSkeletalComponentPrivate()
+#if WITH_EDITOR
+void UCustomizableSkeletalComponent::PostLoad()
 {
-	InstanceUsage = CreateDefaultSubobject<UCustomizableObjectInstanceUsage>(TEXT("InstanceUsage"));
-}
+	Super::PostLoad();
 
+	SetCustomizableObjectInstance(CustomizableObjectInstance);
+	SetComponentName(ComponentName);
+	SetSkipSetReferenceSkeletalMesh(bSkipSetReferenceSkeletalMesh);
+	SetSkipSetSkeletalMeshOnAttach(bSkipSetSkeletalMeshOnAttach);
+}
 
-USkeletalMesh* UCustomizableSkeletalComponentPrivate::GetSkeletalMesh() const
+void UCustomizableSkeletalComponent::PostEditChangeProperty(FPropertyChangedEvent& PropertyChangedEvent)
 {
-	if (InstanceUsage)
+	Super::PostEditChangeProperty(PropertyChangedEvent);
+
+	if (PropertyChangedEvent.Property == nullptr)
+	{
+		return;
+	}
+	if (PropertyChangedEvent.Property->GetFName() == GET_MEMBER_NAME_CHECKED(UCustomizableSkeletalComponent, CustomizableObjectInstance))
+	{
+		SetCustomizableObjectInstance(CustomizableObjectInstance);
+	}
+	else if (PropertyChangedEvent.Property->GetFName() == GET_MEMBER_NAME_CHECKED(UCustomizableSkeletalComponent, ComponentName))
+	{
+		SetComponentName(ComponentName);
+	}
+	else if (PropertyChangedEvent.Property->GetFName() == GET_MEMBER_NAME_CHECKED(UCustomizableSkeletalComponent, bSkipSetReferenceSkeletalMesh))
 	{
-		return InstanceUsage->GetPrivate()->GetSkeletalMesh();
+		SetSkipSetReferenceSkeletalMesh(bSkipSetReferenceSkeletalMesh);
 	}
+	else if (PropertyChangedEvent.Property->GetFName() == GET_MEMBER_NAME_CHECKED(UCustomizableSkeletalComponent, bSkipSetSkeletalMeshOnAttach))
+	{
+		SetSkipSetSkeletalMeshOnAttach(bSkipSetSkeletalMeshOnAttach);
+	}
+}
+#endif
 
-	return nullptr;
+UCustomizableSkeletalComponentPrivate::UCustomizableSkeletalComponentPrivate()
+{
+	InstanceUsage = CreateDefaultSubobject<UCustomizableObjectInstanceUsage>(TEXT("InstanceUsage"));
 }
 
 
-USkeletalMesh* UCustomizableSkeletalComponentPrivate::GetAttachedSkeletalMesh() const
+UCustomizableObjectInstanceUsage* UCustomizableSkeletalComponent::GetInstanceUsage()
 {
-	if (InstanceUsage)
-	{
-		return InstanceUsage->GetPrivate()->GetAttachedSkeletalMesh();
-	}
-
-	return nullptr;
+	return GetPrivate()->InstanceUsage;
 }
 
 
 void UCustomizableSkeletalComponent::SetComponentName(const FName& Name)
 {
 	ComponentName = Name;
+	GetPrivate()->InstanceUsage->SetComponentName(Name);
 }
 
 
 FName UCustomizableSkeletalComponent::GetComponentName() const
 {
-	if (!ComponentName.IsNone())
-	{
-		return ComponentName;	
-	}
-	else
-	{
-		return FName(FString::FromInt(ComponentIndex));
-	}
+	return ComponentName;
 }
 
 
@@ -75,48 +90,33 @@ UCustomizableObjectInstance* UCustomizableSkeletalComponent::GetCustomizableObje
 void UCustomizableSkeletalComponent::SetCustomizableObjectInstance(UCustomizableObjectInstance* Instance)
 {
 	CustomizableObjectInstance = Instance;
+	GetPrivate()->InstanceUsage->SetCustomizableObjectInstance(Instance);
 }
 
 
 void UCustomizableSkeletalComponent::SetSkipSetReferenceSkeletalMesh(bool bSkip)
 {
 	bSkipSetReferenceSkeletalMesh = bSkip;
-}
-
-
-bool UCustomizableSkeletalComponent::GetSkipSetReferenceSkeletalMesh() const
-{
-	return bSkipSetReferenceSkeletalMesh;
+	GetPrivate()->InstanceUsage->SetSkipSetReferenceSkeletalMesh(bSkip);
 }
 
 
 void UCustomizableSkeletalComponent::SetSkipSetSkeletalMeshOnAttach(bool bSkip)
 {
-	bSkipSkipSetSkeletalMeshOnAttach = bSkip;
-}
-
-
-bool UCustomizableSkeletalComponent::GetSkipSetSkeletalMeshOnAttach() const
-{
-	return bSkipSkipSetSkeletalMeshOnAttach;
+	bSkipSetSkeletalMeshOnAttach = bSkip;
+	GetPrivate()->InstanceUsage->SetSkipSetSkeletalMeshOnAttach(bSkip);
 }
 
 
 void UCustomizableSkeletalComponent::UpdateSkeletalMeshAsync(bool bNeverSkipUpdate)
 {
-	if (GetPrivate()->InstanceUsage)
-	{
-		GetPrivate()->InstanceUsage->UpdateSkeletalMeshAsync(bNeverSkipUpdate);
-	}
+	GetPrivate()->InstanceUsage->UpdateSkeletalMeshAsync(bNeverSkipUpdate);
 }
 
 
 void UCustomizableSkeletalComponent::UpdateSkeletalMeshAsyncResult(FInstanceUpdateDelegate Callback, bool bIgnoreCloseDist, bool bForceHighPriority)
 {
-	if (GetPrivate()->InstanceUsage)
-	{
-		GetPrivate()->InstanceUsage->UpdateSkeletalMeshAsyncResult(Callback, bIgnoreCloseDist, bForceHighPriority);
-	}
+	GetPrivate()->InstanceUsage->UpdateSkeletalMeshAsyncResult(Callback, bIgnoreCloseDist, bForceHighPriority);
 }
 
 
@@ -134,24 +134,6 @@ const UCustomizableSkeletalComponentPrivate* UCustomizableSkeletalComponent::Get
 }
 
 
-UCustomizableSkeletalComponent* UCustomizableSkeletalComponentPrivate::GetPublic()
-{
-	UCustomizableSkeletalComponent* Public = StaticCast<UCustomizableSkeletalComponent*>(GetOuter());
-	check(Public);
-
-	return Public;
-}
-
-
-const UCustomizableSkeletalComponent* UCustomizableSkeletalComponentPrivate::GetPublic() const
-{
-	UCustomizableSkeletalComponent* Public = StaticCast<UCustomizableSkeletalComponent*>(GetOuter());
-	check(Public);
-
-	return Public;	
-}
-
-
 void UCustomizableSkeletalComponent::OnAttachmentChanged()
 {
 	Super::OnAttachmentChanged();
diff --git a/Engine/Plugins/Mutable/Source/CustomizableObject/Private/MuCO/CustomizableSkeletalMeshActor.cpp b/Engine/Plugins/Mutable/Source/CustomizableObject/Private/MuCO/CustomizableSkeletalMeshActor.cpp
--- a/Engine/Plugins/Mutable/Source/CustomizableObject/Private/MuCO/CustomizableSkeletalMeshActor.cpp
+++ b/Engine/Plugins/Mutable/Source/CustomizableObject/Private/MuCO/CustomizableSkeletalMeshActor.cpp
@@ -24,7 +24,7 @@ UCustomizableObjectInstance* ACustomizableSkeletalMeshActor::GetCustomizableObje
 {
 	for (const UCustomizableSkeletalComponent* Component : CustomizableSkeletalComponents)
 	{
-		if (UCustomizableObjectInstance* COInstance = Component->CustomizableObjectInstance)
+		if (UCustomizableObjectInstance* COInstance = Component->GetCustomizableObjectInstance())
 		{
 			return COInstance;
 		}
@@ -108,7 +108,7 @@ void UCustomizableSkeletalMeshActorPrivate::Init(UCustomizableObjectInstance* In
 		CustomizableComponent = NewObject<UCustomizableSkeletalComponent>(GetPublic(), UCustomizableSkeletalComponent::StaticClass(), FName(CustomizableComponentName));
 
 		CustomizableComponent->AttachToComponent(Component, FAttachmentTransformRules::KeepRelativeTransform);
-		CustomizableComponent->CustomizableObjectInstance = Instance;
+		CustomizableComponent->SetCustomizableObjectInstance(Instance);
 		CustomizableComponent->SetComponentName(ComponentName);
 		CustomizableComponent->RegisterComponent();
 		
diff --git a/Engine/Plugins/Mutable/Source/CustomizableObject/Public/MuCO/CustomizableObject.h b/Engine/Plugins/Mutable/Source/CustomizableObject/Public/MuCO/CustomizableObject.h
--- a/Engine/Plugins/Mutable/Source/CustomizableObject/Public/MuCO/CustomizableObject.h
+++ b/Engine/Plugins/Mutable/Source/CustomizableObject/Public/MuCO/CustomizableObject.h
@@ -548,6 +548,6 @@ public:
 	
 private:
 
-	UPROPERTY()
+	UPROPERTY(Instanced)
 	TObjectPtr<UCustomizableObjectPrivate> Private;
 };
diff --git a/Engine/Plugins/Mutable/Source/CustomizableObject/Public/MuCO/CustomizableObjectInstance.h b/Engine/Plugins/Mutable/Source/CustomizableObject/Public/MuCO/CustomizableObjectInstance.h
--- a/Engine/Plugins/Mutable/Source/CustomizableObject/Public/MuCO/CustomizableObjectInstance.h
+++ b/Engine/Plugins/Mutable/Source/CustomizableObject/Public/MuCO/CustomizableObjectInstance.h
@@ -9,27 +9,14 @@
 
 #include "CustomizableObjectInstance.generated.h"
 
-namespace UE::Mutable::Private
-{
-	class FImage;
-}
-
-class UCustomizableObjectSystemPrivate;
+class UCustomizableObject;
 class USkeletalMesh;
-class AActor;
-class FProperty;
 class UAnimInstance;
-class UCustomizableObject;
-class UTexture2D;
-class FUpdateContextPrivate;
-class UCustomizableObjectInstanceUsage;
-class UCustomizableObjectExtension;
 class UMaterialInterface;
+class FProperty;
+class FUpdateContextPrivate;
 struct FFrame;
 struct FGameplayTagContainer;
-struct FPropertyChangedEvent;
-struct FTexturePlatformData;
-struct FMutableModelImageProperties;
 
 
 /**
@@ -220,15 +207,6 @@ struct FBakingConfiguration
 };
 
 
-/** Result of all the checks just before beginning an update. */
-enum class EUpdateRequired : uint8
-{
-	NoUpdate, // No work required.
-	Update, // Normal update.
-	Discard // Discard instead of update.
-};
-
-
 /** Instance Update Result. */
 UENUM(BlueprintType)
 enum class EUpdateResult : uint8
@@ -334,7 +312,6 @@ public:
 	CUSTOMIZABLEOBJECT_API virtual void Serialize(FArchive& Ar) override;
 	CUSTOMIZABLEOBJECT_API virtual void PostLoad() override;
 	CUSTOMIZABLEOBJECT_API virtual void BeginDestroy() override;
-	CUSTOMIZABLEOBJECT_API virtual bool IsReadyForFinishDestroy() override;
 	CUSTOMIZABLEOBJECT_API virtual FString GetDesc() override;
 	CUSTOMIZABLEOBJECT_API virtual bool IsEditorOnly() const override;
 	CUSTOMIZABLEOBJECT_API virtual void PostInitProperties() override;
@@ -798,7 +775,7 @@ private:
 	UPROPERTY()
 	FCustomizableObjectInstanceDescriptor Descriptor;
 
-	UPROPERTY()
+	UPROPERTY(Instanced)
 	TObjectPtr<UCustomizableInstancePrivate> PrivateData;
 
 	// Deprecated properties	
diff --git a/Engine/Plugins/Mutable/Source/CustomizableObject/Public/MuCO/CustomizableObjectInstanceUsage.h b/Engine/Plugins/Mutable/Source/CustomizableObject/Public/MuCO/CustomizableObjectInstanceUsage.h
--- a/Engine/Plugins/Mutable/Source/CustomizableObject/Public/MuCO/CustomizableObjectInstanceUsage.h
+++ b/Engine/Plugins/Mutable/Source/CustomizableObject/Public/MuCO/CustomizableObjectInstanceUsage.h
@@ -2,22 +2,15 @@
 
 #pragma once
 
-#include "MuCO/CustomizableObjectInstance.h"
-#include "Tickable.h"
 #include "Engine/EngineBaseTypes.h"
+#include "MuCO/CustomizableObjectInstance.h"
 
 #include "CustomizableObjectInstanceUsage.generated.h"
 
 #define UE_API CUSTOMIZABLEOBJECT_API
 
-class AActor;
-class FObjectPreSaveContext;
 class UObject;
-class USkeletalMesh;
-class UPhysicsAsset;
 class UCustomizableObjectInstanceUsagePrivate;
-struct FFrame;
-enum class EUpdateResult : uint8;
 
 DECLARE_DELEGATE(FCustomizableObjectInstanceUsageUpdatedDelegate);
 
@@ -36,17 +29,19 @@ public:
 	// Own interface
 	UE_API UCustomizableObjectInstanceUsage();
 
-	virtual void PostInitProperties() override;
+	virtual void PostDuplicate(bool bDuplicateForPIE) override;
+	virtual void PostLoad() override;
 	virtual void BeginDestroy() override;
+
+#if WITH_EDITOR
+	UE_API virtual void PostEditChangeProperty(struct FPropertyChangedEvent& PropertyChangedEvent) override;
+#endif
 	
 	UFUNCTION(BlueprintCallable, Category = CustomizableObjectInstanceUsage)
 	UE_API void SetCustomizableObjectInstance(UCustomizableObjectInstance* CustomizableObjectInstance);
 
 	UFUNCTION(BlueprintCallable, Category = CustomizableObjectInstanceUsage)
 	UE_API UCustomizableObjectInstance* GetCustomizableObjectInstance() const;
-	
-	// DEPRECATED: use GetComponentName
-	UE_API int32 GetComponentIndex() const;
 
 	UFUNCTION(BlueprintCallable, Category = CustomizableObjectInstanceUsage)
 	UE_API void SetComponentName(const FName& Name);
@@ -64,7 +59,7 @@ public:
 	
 	/** Update Skeletal Mesh asynchronously. */
 	UFUNCTION(BlueprintCallable, Category = CustomizableObject)
-	UE_API void UpdateSkeletalMeshAsync(bool bNeverSkipUpdate = false, bool bIgnoreCloseDist = false, bool bForceHighPriority = false);
+	UE_API void UpdateSkeletalMeshAsync(bool bNeverSkipUpdate_DEPRECATED = false, bool bIgnoreCloseDist_DEPRECATED = false, bool bForceHighPriority = false);
 
 	/** Update Skeletal Mesh asynchronously. Callback will be called once the update finishes, even if it fails. */
 	UFUNCTION(BlueprintCallable, Category = CustomizableObjectInstance)
@@ -81,9 +76,6 @@ public:
 	/** Set to true to avoid automatically replacing the Skeletal Mesh of the parent Skeletal Mesh Component with any mesh. */
 	UFUNCTION(BlueprintCallable, Category = CustomizableObjectInstanceUsage)
 	UE_API void SetSkipSetSkeletalMeshOnAttach(bool bSkip);
-
-	UFUNCTION(BlueprintCallable, Category = CustomizableObjectInstanceUsage)
-	UE_API bool GetSkipSetSkeletalMeshOnAttach() const;
 	
 	UE_API UCustomizableObjectInstanceUsagePrivate* GetPrivate();
 
@@ -97,28 +89,26 @@ private:
 	TWeakObjectPtr<USkeletalMeshComponent> UsedSkeletalMeshComponent;
 
 	// If the outer is no CustomizableSkeletalComponentd, this Instance will be used
-	UPROPERTY()
+	UPROPERTY(EditAnywhere, Category = CustomizableObjectInstanceUsage, meta = (DisplayName = "Customizable Object Instance"))
 	TObjectPtr<UCustomizableObjectInstance> UsedCustomizableObjectInstance;
 
-	// If the outer is no CustomizableSkeletalComponentd, this Index will be used
-	// DEPRECATED
-	UPROPERTY()
-	int32 UsedComponentIndex;
-
-	/** Only used if the UsedComponentIndex is INDEX_NONE. */
-	UPROPERTY()
+	UPROPERTY(EditAnywhere, Category = CustomizableObjectInstanceUsage, meta = (DisplayName = "Component Name"))
 	FName UsedComponentName;
 
 	// Used to avoid replacing the SkeletalMesh of the parent component by the ReferenceSkeletalMesh if bPendingSetSkeletalMesh is true
-	UPROPERTY()
+	UPROPERTY(EditAnywhere, Category = CustomizableObjectInstanceUsage, meta = (DisplayName = "Skip Set Reference Skeletal Mesh"))
 	bool bUsedSkipSetReferenceSkeletalMesh = false;
 
-	UPROPERTY()
+	UPROPERTY(EditAnywhere, Category = CustomizableObjectInstanceUsage, meta = (DisplayName = "Skip Set Skeletal Mesh On Attach"))
 	bool bUsedSkipSetSkeletalMeshOnAttach = false;
 	
 	UPROPERTY(Instanced)
 	TObjectPtr<UCustomizableObjectInstanceUsagePrivate> Private;
 	
+#if WITH_EDITORONLY_DATA
+	TWeakObjectPtr<UCustomizableObjectInstance> OldCustomizableObjectInstance;
+#endif
+
 	friend UCustomizableObjectInstanceUsagePrivate;
 };
 
diff --git a/Engine/Plugins/Mutable/Source/CustomizableObject/Public/MuCO/CustomizableSkeletalComponent.h b/Engine/Plugins/Mutable/Source/CustomizableObject/Public/MuCO/CustomizableSkeletalComponent.h
--- a/Engine/Plugins/Mutable/Source/CustomizableObject/Public/MuCO/CustomizableSkeletalComponent.h
+++ b/Engine/Plugins/Mutable/Source/CustomizableObject/Public/MuCO/CustomizableSkeletalComponent.h
@@ -5,13 +5,14 @@
 #include "Components/SceneComponent.h"
 
 #include "MuCO/CustomizableObjectInstance.h"
-#include "MuCO/CustomizableObject.h"
+
 
 #include "CustomizableSkeletalComponent.generated.h"
 
 #define UE_API CUSTOMIZABLEOBJECT_API
 
 class UCustomizableSkeletalComponentPrivate;
+class UCustomizableObjectInstanceUsage;
 
 DECLARE_DELEGATE(FCustomizableSkeletalComponentUpdatedDelegate);
 
@@ -21,49 +22,45 @@ class UCustomizableSkeletalComponent : public USceneComponent
 {
 	friend UCustomizableSkeletalComponentPrivate;
 
-public:
 	GENERATED_BODY()
-	
-	UPROPERTY(BlueprintReadWrite, EditAnywhere, Category = CustomizableSkeletalComponent)
-	TObjectPtr<UCustomizableObjectInstance> CustomizableObjectInstance;
 
-	/** This component index refers to the object list of components 
+public:
+	/** This component index refers to the object list of components
 	* DEPRECATED: Use FNames instead with Get/SetComponentName
 	*/
 	UPROPERTY(BlueprintReadWrite, Category = CustomizableSkeletalComponent)
 	int32 ComponentIndex;
 
 private:
-	/** Only used if the ComponentIndex is INDEX_NONE.
-	 *  Editing this property will set ComponentIndex to INDEX_NONE. */
+
+	UPROPERTY(EditAnywhere, Category = CustomizableSkeletalComponent)
+	TObjectPtr<UCustomizableObjectInstance> CustomizableObjectInstance;
+
 	UPROPERTY(EditAnywhere, Category = CustomizableSkeletalComponent)
 	FName ComponentName;
 
 	UPROPERTY(EditAnywhere, Category = CustomizableSkeletalComponent)
 	bool bSkipSetReferenceSkeletalMesh = false;
-	
-public:
-	FCustomizableSkeletalComponentUpdatedDelegate UpdatedDelegate;
 
-private:
-	UPROPERTY()
-	bool bSkipSkipSetSkeletalMeshOnAttach = false;
+	UPROPERTY(EditAnywhere, Category = CustomizableSkeletalComponent)
+	bool bSkipSetSkeletalMeshOnAttach = false;
 	
 	UPROPERTY(Instanced)
 	TObjectPtr<UCustomizableSkeletalComponentPrivate> Private;
 
 public:
 	
+	// Own interface
 	UE_API UCustomizableSkeletalComponent();
-	
-protected:
 
-	// USceneComponent interface
-	UE_API virtual void OnAttachmentChanged() override;
+#if WITH_EDITOR
+	UE_API virtual void PostLoad();
+	UE_API virtual void PostEditChangeProperty(struct FPropertyChangedEvent& PropertyChangedEvent) override;
+#endif
+
+	UFUNCTION(BlueprintCallable, Category = CustomizableSkeletalComponent)
+	UE_API UCustomizableObjectInstanceUsage* GetInstanceUsage();
 
-public:
-	// Own interface
-	
 	UFUNCTION(BlueprintCallable, Category = CustomizableSkeletalComponent)
 	UE_API void SetComponentName(const FName& Name);
 
@@ -81,27 +78,26 @@ public:
 	UFUNCTION(BlueprintCallable, Category = CustomizableSkeletalComponent)
 	UE_API void SetSkipSetReferenceSkeletalMesh(bool bSkip);
 
-	UFUNCTION(BlueprintCallable, Category = CustomizableSkeletalComponent)
-	UE_API bool GetSkipSetReferenceSkeletalMesh() const;
-
 	/** Set to true to avoid automatically replacing the Skeletal Mesh of the parent Skeletal Mesh Component with any mesh. */
 	UFUNCTION(BlueprintCallable, Category = CustomizableObjectInstanceUsage)
 	UE_API void SetSkipSetSkeletalMeshOnAttach(bool bSkip);
-
-	UFUNCTION(BlueprintCallable, Category = CustomizableObjectInstanceUsage)
-	UE_API bool GetSkipSetSkeletalMeshOnAttach() const;
 	
 	/** Update Skeletal Mesh asynchronously. */
 	UFUNCTION(BlueprintCallable, Category = CustomizableSkeletalComponent)
-	UE_API void UpdateSkeletalMeshAsync(bool bNeverSkipUpdate = false);
+	UE_API void UpdateSkeletalMeshAsync(bool bForceUpdate = false);
 
 	/** Update Skeletal Mesh asynchronously. Callback will be called once the update finishes, even if it fails. */
 	UFUNCTION(BlueprintCallable, Category = CustomizableSkeletalComponent)
 	UE_API void UpdateSkeletalMeshAsyncResult(FInstanceUpdateDelegate Callback, bool bIgnoreCloseDist = false, bool bForceHighPriority = false);
 
+protected:
+
 	UE_API UCustomizableSkeletalComponentPrivate* GetPrivate();
 
 	UE_API const UCustomizableSkeletalComponentPrivate* GetPrivate() const;
+
+	// USceneComponent interface
+	UE_API virtual void OnAttachmentChanged() override;
 };
 
 #undef UE_API
diff --git a/Engine/Plugins/Mutable/Source/CustomizableObjectEditor/Private/MuCOE/CustomizableObjectEditorModule.cpp b/Engine/Plugins/Mutable/Source/CustomizableObjectEditor/Private/MuCOE/CustomizableObjectEditorModule.cpp
--- a/Engine/Plugins/Mutable/Source/CustomizableObjectEditor/Private/MuCOE/CustomizableObjectEditorModule.cpp
+++ b/Engine/Plugins/Mutable/Source/CustomizableObjectEditor/Private/MuCOE/CustomizableObjectEditorModule.cpp
@@ -12,10 +12,8 @@
 #include "Misc/DateTime.h"
 #include "Misc/Timespan.h"
 #include "HAL/FileManager.h"
-#include "LevelEditorViewport.h"
 #include "MuCO/CustomizableObjectSystem.h"		// For defines related to memory function replacements.
 #include "MuCO/CustomizableObjectInstanceUsage.h"
-#include "MuCO/CustomizableObjectInstanceUsagePrivate.h"
 #include "MuCO/ICustomizableObjectModule.h"		// For instance editor command utility function
 #include "MuCO/UnrealPortabilityHelpers.h"
 #include "MuCOE/CustomizableInstanceDetails.h"
@@ -130,7 +128,7 @@ constexpr int32 MaxAccessTimespan = 30;
 
 TAutoConsoleVariable<bool> CVarMutableDerivedDataCacheUsage(
 	TEXT("mutable.DerivedDataCacheUsage"),
-	true,
+	false,
 	TEXT("Derived data cache access for compiled data.")
 	TEXT("true - Allows DDC access according to the DDC policies specified in the plugin settings."),
 	ECVF_Default);
@@ -1361,58 +1359,6 @@ void FCustomizableObjectEditorModule::CompileCustomizableObject(UCustomizableObj
 }
 
 
-void FCustomizableObjectEditorModule::UpdateInstanceDistanceToLevelEditorCamera()
-{
-	MUTABLE_EDITOR_CPUPROFILER_SCOPE(UpdateInstanceDistanceToLevelEditorCamera);
-
-	for (TObjectIterator<UCustomizableObjectInstanceUsage> CustomizableObjectInstanceUsage; CustomizableObjectInstanceUsage; ++CustomizableObjectInstanceUsage)
-	{
-		if (!IsValid(*CustomizableObjectInstanceUsage) || CustomizableObjectInstanceUsage->IsTemplate())
-		{
-			continue;
-		}
-
-		if (CustomizableObjectInstanceUsage->GetPrivate()->bIsNetModeDedicatedServer)
-		{
-			continue;
-		}
-
-		UCustomizableObjectInstance* CustomizableObjectInstance = CustomizableObjectInstanceUsage->GetCustomizableObjectInstance();
-		if (!CustomizableObjectInstance)
-		{
-			continue;
-		}
-
-		USkeletalMeshComponent* Parent = Cast<USkeletalMeshComponent>(CustomizableObjectInstanceUsage->GetAttachParent());
-		if (!Parent)
-		{
-			continue;
-		}
-		
-		UWorld* LocalWorld = Parent->GetWorld();
-		if (!LocalWorld || LocalWorld->WorldType != EWorldType::Editor)
-		{
-			continue;
-		}
-
-		AActor* ParentActor = Parent->GetAttachmentRootActor();
-		if (ParentActor)
-		{
-			for (FLevelEditorViewportClient* LevelVC : GEditor->GetLevelViewportClients())
-			{
-				if (LevelVC && LevelVC->IsPerspective() && LevelVC->GetWorld() == LocalWorld)
-				{
-					double SquareDist = FVector::DistSquared(LevelVC->GetViewLocation(), ParentActor->GetActorLocation());
-					CustomizableObjectInstance->GetPrivate()->MinSquareDistFromComponentToPlayer =
-						FMath::Min(SquareDist, CustomizableObjectInstance->GetPrivate()->MinSquareDistFromComponentToPlayer);
-					break;
-				}
-			}
-		}
-	}
-}
-
-
 void FCustomizableObjectEditorModule::Tick(bool bBlocking)
 {
 	if (TSet<TStrongObjectPtr<UCustomizableObject>>::TIterator It = ObjectsPendingLoad.CreateIterator())
diff --git a/Engine/Plugins/Mutable/Source/CustomizableObjectEditor/Private/MuCOE/CustomizableObjectEditorModule.h b/Engine/Plugins/Mutable/Source/CustomizableObjectEditor/Private/MuCOE/CustomizableObjectEditorModule.h
--- a/Engine/Plugins/Mutable/Source/CustomizableObjectEditor/Private/MuCOE/CustomizableObjectEditorModule.h
+++ b/Engine/Plugins/Mutable/Source/CustomizableObjectEditor/Private/MuCOE/CustomizableObjectEditorModule.h
@@ -67,7 +67,6 @@ public:
 	virtual int32 GetNumCompileRequests() override;
 	virtual bool IsCompiling(const UCustomizableObject& Object) const override;
 	virtual void CompileCustomizableObject(UCustomizableObject& Object, const FCompileParams* Params, bool bSilent, bool bForce) override;
-	virtual void UpdateInstanceDistanceToLevelEditorCamera() override;
 	virtual void Tick(bool bBlocking) override;
 	virtual int32 GetRemainingWork() const override;
 	virtual void BeginCacheForCookedPlatformData(UCustomizableObject& Object, const ITargetPlatform* TargetPlatform) override;
diff --git a/Engine/Plugins/Mutable/Source/CustomizableObjectEditor/Private/MuCOE/CustomizableObjectEditorViewportClient.cpp b/Engine/Plugins/Mutable/Source/CustomizableObjectEditor/Private/MuCOE/CustomizableObjectEditorViewportClient.cpp
--- a/Engine/Plugins/Mutable/Source/CustomizableObjectEditor/Private/MuCOE/CustomizableObjectEditorViewportClient.cpp
+++ b/Engine/Plugins/Mutable/Source/CustomizableObjectEditor/Private/MuCOE/CustomizableObjectEditorViewportClient.cpp
@@ -905,7 +905,7 @@ void FCustomizableObjectEditorViewportClient::OnPreSetSkeletalMesh(const FPreSet
 
 				UCustomizableSkeletalComponent* CustomizableComponent = NewObject<UCustomizableSkeletalComponent>(DebugComponent);
 				CustomizableComponent->SetSkipSetReferenceSkeletalMesh(true);
-				CustomizableComponent->CustomizableObjectInstance = Instance;
+				CustomizableComponent->SetCustomizableObjectInstance(Instance);
 				CustomizableComponent->SetComponentName(Name);
 				CustomizableComponent->AttachToComponent(DebugComponent, FAttachmentTransformRules::KeepRelativeTransform);
 				CustomizableComponent->RegisterComponent();
