[Mutable] Avoid invalidating Instance generated data when changing the Customizable Object.

- Also with SetDescriptor.

#rnx

#Synced-CL 52129579

#Synced-CL 52129579

diff --git a/Engine/Plugins/Mutable/Source/CustomizableObject/Internal/MuCO/CustomizableObjectInstancePrivate.h b/Engine/Plugins/Mutable/Source/CustomizableObject/Internal/MuCO/CustomizableObjectInstancePrivate.h
--- a/Engine/Plugins/Mutable/Source/CustomizableObject/Internal/MuCO/CustomizableObjectInstancePrivate.h
+++ b/Engine/Plugins/Mutable/Source/CustomizableObject/Internal/MuCO/CustomizableObjectInstancePrivate.h
@@ -240,7 +240,7 @@ public:
 
 	UE_API FCustomizableObjectInstanceDescriptor& GetDescriptor() const;
 
-	UE_API void SetDescriptor(const FCustomizableObjectInstanceDescriptor& InDescriptor);
+	UE_API void CopyParametersFromInstance(UCustomizableObjectInstance* Instance);
 
 	/** Return true if the instance is not locked and if it's compiled. */
 	UE_API bool CanUpdateInstance() const;
diff --git a/Engine/Plugins/Mutable/Source/CustomizableObject/Private/MuCO/CustomizableObjectInstance.cpp b/Engine/Plugins/Mutable/Source/CustomizableObject/Private/MuCO/CustomizableObjectInstance.cpp
--- a/Engine/Plugins/Mutable/Source/CustomizableObject/Private/MuCO/CustomizableObjectInstance.cpp
+++ b/Engine/Plugins/Mutable/Source/CustomizableObject/Private/MuCO/CustomizableObjectInstance.cpp
@@ -149,15 +149,23 @@ UCustomizableObjectInstance::UCustomizableObjectInstance()
 }
 
 
-void UCustomizableInstancePrivate::SetDescriptor(const FCustomizableObjectInstanceDescriptor& InDescriptor)
+void UCustomizableInstancePrivate::CopyParametersFromInstance(UCustomizableObjectInstance* Instance)
 {
+	if (!Instance)
+	{
+		return;
+	}
+	
+	FCustomizableObjectInstanceDescriptor& Descriptor = Instance->GetPrivate()->GetDescriptor();
+	
+	UCustomizableObject* InCustomizableObject = Descriptor.GetCustomizableObject();
+
 #if WITH_EDITOR
 	// Bind a lambda to the PostCompileDelegate and unbind from the previous object if any.
-	BindObjectDelegates(GetPublic()->GetCustomizableObject(), InDescriptor.GetCustomizableObject());
+	BindObjectDelegates(GetPublic()->GetCustomizableObject(), Descriptor.GetCustomizableObject());
 #endif
 
-	GetPublic()->Descriptor = InDescriptor;
-	InvalidateGeneratedData();
+	GetPublic()->Descriptor = Descriptor;
 }
 
 
@@ -415,7 +423,6 @@ void UCustomizableObjectInstance::SetObject(UCustomizableObject* InObject)
 #endif
 
 	Descriptor.SetCustomizableObject(InObject);
-	GetPrivate()->InvalidateGeneratedData();
 }
 
 
@@ -1714,7 +1721,7 @@ bool UCustomizableInstancePrivate::UpdateSkeletalMesh_PostBeginUpdate0(UCustomiz
 
 	bool bUpdateMeshes = DoComponentsNeedUpdate(Context, bHasInvalidMesh);
 
-	UCustomizableObject* CustomizableObject = Instance->GetCustomizableObject();
+	UCustomizableObject* CustomizableObject = Context->Object.Get();
 
 	if (!CustomizableObject)
 	{
@@ -1832,7 +1839,7 @@ bool UCustomizableInstancePrivate::UpdateSkeletalMesh_PostBeginUpdate0(UCustomiz
 				}
 			}
 
-			FString SkeletalMeshName = FString::Printf(TEXT("SK_%s_%s_h%lu"), *Instance->GetCustomizableObject()->GetName(), *ComponentName.ToString(), MeshHash);
+			FString SkeletalMeshName = FString::Printf(TEXT("SK_%s_%s_h%lu"), *CustomizableObject->GetName(), *ComponentName.ToString(), MeshHash);
 
 #if WITH_EDITOR
 			if (Context->bBake)
@@ -2001,7 +2008,7 @@ UCustomizableObjectInstance* UCustomizableObjectInstance::CloneStatic(UObject* O
 
 void UCustomizableObjectInstance::CopyParametersFromInstance(UCustomizableObjectInstance* Instance)
 {
-	GetPrivate()->SetDescriptor(Instance->GetPrivate()->GetDescriptor());
+	GetPrivate()->CopyParametersFromInstance(Instance);
 }
 
 
@@ -3860,7 +3867,7 @@ bool BuildOrCopyRenderData(const TSharedRef<FUpdateContextPrivate>& Context, USk
 	FSkeletalMeshRenderData* RenderData = SkeletalMesh->GetResourceForRendering();
 	check(RenderData);
 
-	UCustomizableObject* CustomizableObject = Public->GetCustomizableObject();
+	UCustomizableObject* CustomizableObject = Context->Object.Get();
 	
 	const TSharedPtr<FInstanceUpdateData::FComponent>& Component = Context->InstanceUpdateData.Components[ComponentIndex];
 	if (!Component)
@@ -4593,7 +4600,7 @@ void UCustomizableInstancePrivate::BuildResources(const TSharedRef<FUpdateContex
 {
 	MUTABLE_CPUPROFILER_SCOPE(UCustomizableInstancePrivate::BuildResources)
 
-	UCustomizableObject* CustomizableObject = Public->GetCustomizableObject();
+	UCustomizableObject* CustomizableObject = Context->Object.Get();
 
 	const UModelResources& ModelResources = *CustomizableObject->GetPrivate()->GetModelResources();
 	
@@ -4941,12 +4948,6 @@ void UCustomizableObjectInstance::SetKeepOwnershipOfGeneratedResources(bool bEna
 }
 
 
-int32 UCustomizableObjectInstance::GetNumComponents() const
-{
-	return GetCustomizableObject() ? GetCustomizableObject()->GetComponentCount() : 0;
-}
-
-
 #if WITH_EDITOR
 void UCustomizableObjectInstance::Bake(const FBakingConfiguration& InBakingConfiguration)
 {
diff --git a/Engine/Plugins/Mutable/Source/CustomizableObject/Private/MuCO/CustomizableObjectSystem.cpp b/Engine/Plugins/Mutable/Source/CustomizableObject/Private/MuCO/CustomizableObjectSystem.cpp
--- a/Engine/Plugins/Mutable/Source/CustomizableObject/Private/MuCO/CustomizableObjectSystem.cpp
+++ b/Engine/Plugins/Mutable/Source/CustomizableObject/Private/MuCO/CustomizableObjectSystem.cpp
@@ -2359,7 +2359,7 @@ namespace impl
 			return;
 		}
 
-		UCustomizableObject* CustomizableObject = CustomizableObjectInstance->GetCustomizableObject();
+		UCustomizableObject* CustomizableObject = Context->Object.Get();
 		if (!IsValid(CustomizableObject))
 		{
 			Context->UpdateResult = EUpdateResult::Error;
@@ -2605,7 +2605,7 @@ namespace impl
 		
 		SystemPrivate->CurrentInstanceBeingUpdated = Instance;
 
-		const UModelResources& ModelResources = Instance->GetCustomizableObject()->GetPrivate()->GetModelResourcesChecked();
+		const UModelResources& ModelResources = Operation->Object->GetPrivate()->GetModelResourcesChecked();
 
 		UCustomizableInstancePrivate* InstancePrivate = Instance->GetPrivate();
 		check(InstancePrivate); // Already checked in GetPrivate. Make static analyzer happy.
@@ -2667,7 +2667,7 @@ namespace impl
 					return;
 				}
 				
-				UCustomizableObject* Object = Operation->GetCapturedDescriptor().GetCustomizableObject();
+				UCustomizableObject* Object = Operation->Object.Get();
 				if (!Object)
 				{
 					Event.Trigger();
@@ -2700,7 +2700,7 @@ namespace impl
 
 			// Prepare streaming for the current customizable object
 			check(SystemPrivate->Streamer != nullptr);
-			SystemPrivate->Streamer->PrepareStreamingForObject(Operation->Instance->GetCustomizableObject());			
+			SystemPrivate->Streamer->PrepareStreamingForObject(Operation->Object.Get());			
 
 			Operation->bLowPriorityTasksBlocked = true;
 			SystemPrivate->MutableTaskGraph.AllowLaunchingMutableTaskLowPriority(false, false);
diff --git a/Engine/Plugins/Mutable/Source/CustomizableObject/Public/MuCO/CustomizableObjectInstance.h b/Engine/Plugins/Mutable/Source/CustomizableObject/Public/MuCO/CustomizableObjectInstance.h
--- a/Engine/Plugins/Mutable/Source/CustomizableObject/Public/MuCO/CustomizableObjectInstance.h
+++ b/Engine/Plugins/Mutable/Source/CustomizableObject/Public/MuCO/CustomizableObjectInstance.h
@@ -401,6 +401,7 @@ public:
 	CUSTOMIZABLEOBJECT_API UCustomizableObjectInstance* CloneStatic(UObject* Outer);
 
 	/** Copy parameters from the given Instance. */
+	UFUNCTION(BlueprintCallable, Category = CustomizableObjectInstance)
 	CUSTOMIZABLEOBJECT_API void CopyParametersFromInstance(UCustomizableObjectInstance* Instance);
 
 	/** Destroy the Mutable Core Live Update Instance attached to this (if exists). */
@@ -752,13 +753,6 @@ public:
 	/** If true, Mutable will keep the ownership of the generated resources. This means that Mutable may decide to modify or reuse them for performance considerations. */
 	UFUNCTION(BlueprintCallable, Category = CustomizableObjectInstance)
 	CUSTOMIZABLEOBJECT_API void SetKeepOwnershipOfGeneratedResources(bool bEnabled);
-
-	/** Return the total number of components that can be generated with the CustomizableObject (unrelated to this instance parameters). 
-	* DEPRECATED: Get the CO with GetCustomizableObect and get the components there with GetComponentCount.
-	* If the actual components of this instance are required, use GetComponents to get a list of names.
-	*/
-	UFUNCTION(BlueprintCallable, Category = CustomizableObjectInstance)
-	CUSTOMIZABLEOBJECT_API int32 GetNumComponents() const;
 	
 #if WITH_EDITOR
 	/**
