[Mutable] Fix CodeRunner double Material load.

- This is not allowed as loading from the cache decrements the OpHitCount.
- Added GetMaterialIfSet function for reading matrials form the cache without decrementing.

#rnx

#Synced-CL 52099427

diff --git a/Engine/Plugins/Mutable/Source/MutableRuntime/Internal/MuR/ProgramCache.h b/Engine/Plugins/Mutable/Source/MutableRuntime/Internal/MuR/ProgramCache.h
--- a/Engine/Plugins/Mutable/Source/MutableRuntime/Internal/MuR/ProgramCache.h
+++ b/Engine/Plugins/Mutable/Source/MutableRuntime/Internal/MuR/ProgramCache.h
@@ -128,9 +128,10 @@ namespace UE::Mutable::Private
 		void SetAborted(FCacheAddress Address);
 		bool IsAborted(FCacheAddress Address);
 
-		[[nodiscard]] UE_API bool LoadBool(FCacheAddress Address);
-		[[nodiscard]] UE_API bool GetBoolIfSet(FCacheAddress Address, bool& bOutResult);
+		[[nodiscard]] UE_API bool GetBoolIfSet(FCacheAddress Address, bool& bOutValue);
+		[[nodiscard]] UE_API bool GetMaterialIfSet(FCacheAddress Address, TManagedPtr<const FMaterial>& bOutValue);
 
+		[[nodiscard]] UE_API bool LoadBool(FCacheAddress Address);
 		[[nodiscard]] UE_API float LoadScalar(FCacheAddress Address);
 		[[nodiscard]] UE_API int32 LoadInt(FCacheAddress Address);
 		[[nodiscard]] UE_API FVector4f LoadColor(FCacheAddress Address);
@@ -175,7 +176,7 @@ namespace UE::Mutable::Private
 		 * Tries to acquire any weak result. Returns if the operation needs to schedule.
 		 *
 		 * If the result is true, it is mandatory to to execute the operation at some point,
-		 * otherwise other accesses to the same adress will never complete.
+		 * otherwise other accesses to the same address will never complete.
 		 */
 		[[nodiscard]] bool AcquireCachedResultForScheduling(FCacheAddress Address);
 
diff --git a/Engine/Plugins/Mutable/Source/MutableRuntime/Private/MuR/CodeRunner.cpp b/Engine/Plugins/Mutable/Source/MutableRuntime/Private/MuR/CodeRunner.cpp
--- a/Engine/Plugins/Mutable/Source/MutableRuntime/Private/MuR/CodeRunner.cpp
+++ b/Engine/Plugins/Mutable/Source/MutableRuntime/Private/MuR/CodeRunner.cpp
@@ -5607,7 +5607,10 @@ namespace UE::Mutable::Private
 			}
 			case 1:
 			{
-				TManagedPtr<const FMaterial> Material = LoadMaterial(FCacheAddress(Args.Material, Item));
+				TManagedPtr<const FMaterial> Material;
+
+				bool bSuccess = GetMemory().GetMaterialIfSet(FCacheAddress(Args.Material, Item), Material);
+				check(bSuccess);
 
 				const FName& ParameterName = Program.ConstantNames[Args.ParameterName];
 				const FParameterKey ParameterKey = { ParameterName, Args.LayerIndex };
diff --git a/Engine/Plugins/Mutable/Source/MutableRuntime/Private/MuR/ProgramCache.cpp b/Engine/Plugins/Mutable/Source/MutableRuntime/Private/MuR/ProgramCache.cpp
--- a/Engine/Plugins/Mutable/Source/MutableRuntime/Private/MuR/ProgramCache.cpp
+++ b/Engine/Plugins/Mutable/Source/MutableRuntime/Private/MuR/ProgramCache.cpp
@@ -1382,21 +1382,30 @@ namespace ProgramCacheInternals
 		check(bSuccess);
 	}
 
-	bool FProgramCache::LoadBool(FCacheAddress Address)
+	bool FProgramCache::GetBoolIfSet(FCacheAddress Address, bool& OutResult)
 	{
 		check(Address.At == 0 || GetAddressDataType(Address.At) == EDataType::Bool);
 
-		bool Value = false;
-		GetInternals()->BoolCache.GetValue(Address, Value);
+		OutResult = false;
+		return GetInternals()->BoolCache.GetValue(Address, OutResult, false);
+	}
 
-		return Value;
+	bool FProgramCache::GetMaterialIfSet(FCacheAddress Address, TManagedPtr<const FMaterial>& OutValue)
+	{
+		check(Address.At == 0 || GetAddressDataType(Address.At) == EDataType::Material);
+
+		OutValue = nullptr;
+		return GetInternals()->MaterialCache.GetValue(Address, OutValue, false);
 	}
 
-	bool FProgramCache::GetBoolIfSet(FCacheAddress Address, bool& OutResult)
+	bool FProgramCache::LoadBool(FCacheAddress Address)
 	{
 		check(Address.At == 0 || GetAddressDataType(Address.At) == EDataType::Bool);
 
-		return GetInternals()->BoolCache.GetValue(Address, OutResult, false);
+		bool Value = false;
+		GetInternals()->BoolCache.GetValue(Address, Value);
+
+		return Value;
 	}
 
 	float FProgramCache::LoadScalar(FCacheAddress Address)
@@ -1692,11 +1701,6 @@ namespace ProgramCacheInternals
 			TargetBytesToFree -= GetInternals()->MaterialCache.TryFreeWeakMemory(TargetBytesToFree);
 		}
 
-		if (TargetBytesToFree > 0)
-		{
-			TargetBytesToFree -= GetInternals()->MaterialCache.TryFreeWeakMemory(TargetBytesToFree);
-		}
-
 		if (TargetBytesToFree > 0)
 		{
 			TargetBytesToFree -= GetInternals()->SkeletalMeshCache.TryFreeWeakMemory(TargetBytesToFree);
