[Mutable] Fix compilation hangs due to oversubscription.

Converted GenerateExternal to Tasks.

#Synced-CL 52129579

diff --git a/Engine/Plugins/Mutable/Source/MutableTools/Internal/MuT/AST.h b/Engine/Plugins/Mutable/Source/MutableTools/Internal/MuT/AST.h
--- a/Engine/Plugins/Mutable/Source/MutableTools/Internal/MuT/AST.h
+++ b/Engine/Plugins/Mutable/Source/MutableTools/Internal/MuT/AST.h
@@ -457,6 +457,8 @@ namespace UE::Mutable::Private
     class ASTChild
     {
     public:
+    	struct FInvalidTag {};
+    	
         explicit ASTChild(ASTOp* parent, const Ptr<ASTOp>& child=Ptr<ASTOp>());
         ASTChild(const Ptr<ASTOp>& parent, const Ptr<ASTOp>& child);
         ~ASTChild();
@@ -464,7 +466,9 @@ namespace UE::Mutable::Private
         ASTChild(const ASTChild&) = delete;
         ASTChild& operator=(const ASTChild&) = delete;
 
-        // move constructor
+    	explicit ASTChild(FInvalidTag) {}
+
+        // move constructor    	
         ASTChild(ASTChild&& rhs)
              : Parent(rhs.Parent)
              , Child(rhs.Child)
diff --git a/Engine/Plugins/Mutable/Source/MutableTools/Private/MuT/ASTOpSwitch.h b/Engine/Plugins/Mutable/Source/MutableTools/Private/MuT/ASTOpSwitch.h
--- a/Engine/Plugins/Mutable/Source/MutableTools/Private/MuT/ASTOpSwitch.h
+++ b/Engine/Plugins/Mutable/Source/MutableTools/Private/MuT/ASTOpSwitch.h
@@ -33,6 +33,10 @@ template <class SCALAR> class vec4;
 
 		struct FCase
 		{
+			struct FInvalidTag {};
+
+			explicit FCase(FInvalidTag) : Branch(ASTChild::FInvalidTag()) {}
+			
 			FCase(int32 cond, Ptr<ASTOp> parent, Ptr<ASTOp> b)
 				: Condition(cond)
 				, Branch(parent.get(), b)
diff --git a/Engine/Plugins/Mutable/Source/MutableTools/Private/MuT/CodeGenerator.h b/Engine/Plugins/Mutable/Source/MutableTools/Private/MuT/CodeGenerator.h
--- a/Engine/Plugins/Mutable/Source/MutableTools/Private/MuT/CodeGenerator.h
+++ b/Engine/Plugins/Mutable/Source/MutableTools/Private/MuT/CodeGenerator.h
@@ -993,15 +993,7 @@ namespace UE::Mutable::Private
     		bool operator==(const FGeneratedExternalCacheKey& Other) const = default;
     	};
     	
-    	struct FExternalGenerationResult
-    	{
-    		Ptr<ASTOp> Op;
-    		
-    		// If set, at least one of the external operation inputs is a mesh
-    		TOptional<FMeshGenerationResult> MeshGenerationResult;
-    	};
-    	
-    	typedef TMap<FGeneratedExternalCacheKey, FExternalGenerationResult> FGeneratedExternalMap;
+    	typedef TMap<FGeneratedExternalCacheKey, FExternalTask> FGeneratedExternalMap;
     	struct FGeneratedExternal
     	{
     		UE::FMutex Mutex;
@@ -1009,10 +1001,10 @@ namespace UE::Mutable::Private
     	};
     	FGeneratedExternal GeneratedExternal;
 
-    	void GenerateExternal(FExternalGenerationResult& OutResult, const FExternalGenerationOptions& Options, const Ptr<const NodeExternal>& ExternalNode);
-    	void GenerateExternal_Operation(FExternalGenerationResult& OutResult, const FExternalGenerationOptions& Options, const Ptr<const NodeExternalOperation>& ExternalNode);
-    	void GenerateExternal_Parameter(FExternalGenerationResult& OutResult, const FExternalGenerationOptions& Options, const Ptr<const NodeExternalParameter>& ExtensionNode);
-    	void GenerateExternal_Switch(FExternalGenerationResult& OutResult, const FExternalGenerationOptions& Options, const Ptr<const NodeExternalSwitch>& ExtensionNode);
+    	FExternalTask GenerateExternal(const FExternalGenerationOptions& Options, const Ptr<const NodeExternal>& ExternalNode);
+    	FExternalTask GenerateExternal_Operation(const FExternalGenerationOptions& Options, const Ptr<const NodeExternalOperation>& ExternalNode);
+    	FExternalTask GenerateExternal_Parameter(const FExternalGenerationOptions& Options, const Ptr<const NodeExternalParameter>& ExtensionNode);
+    	FExternalTask GenerateExternal_Switch(const FExternalGenerationOptions& Options, const Ptr<const NodeExternalSwitch>& ExtensionNode);
     	Ptr<ASTOp> GenerateMissingExternalCode(const TCHAR* StrWhere, const void* ErrorContext);
     	
 		//-----------------------------------------------------------------------------------------
diff --git a/Engine/Plugins/Mutable/Source/MutableTools/Private/MuT/CodeGenerator_Color.cpp b/Engine/Plugins/Mutable/Source/MutableTools/Private/MuT/CodeGenerator_Color.cpp
--- a/Engine/Plugins/Mutable/Source/MutableTools/Private/MuT/CodeGenerator_Color.cpp
+++ b/Engine/Plugins/Mutable/Source/MutableTools/Private/MuT/CodeGenerator_Color.cpp
@@ -577,10 +577,10 @@ namespace UE::Mutable::Private
 		FExternalGenerationOptions Options;
 		Options.ColourOptions = InOptions;
 					
-		FExternalGenerationResult ExtensionResult;
-		GenerateExternal(ExtensionResult, Options, InNode->Node);
-
-		Result.op = ExtensionResult.Op;
+		FExternalTask ExtensionResult = GenerateExternal(Options, InNode->Node);
+		WaitTask(ExtensionResult);
+		
+		Result.op = ExtensionResult.GetResult().Op;
 	}
 
 
diff --git a/Engine/Plugins/Mutable/Source/MutableTools/Private/MuT/CodeGenerator_External.cpp b/Engine/Plugins/Mutable/Source/MutableTools/Private/MuT/CodeGenerator_External.cpp
--- a/Engine/Plugins/Mutable/Source/MutableTools/Private/MuT/CodeGenerator_External.cpp
+++ b/Engine/Plugins/Mutable/Source/MutableTools/Private/MuT/CodeGenerator_External.cpp
@@ -15,12 +15,11 @@
 
 namespace UE::Mutable::Private
 {
-	void CodeGenerator::GenerateExternal(FExternalGenerationResult& OutResult, const FExternalGenerationOptions& Options, const Ptr<const NodeExternal>& ExternalNode)
+	FExternalTask CodeGenerator::GenerateExternal(const FExternalGenerationOptions& Options, const Ptr<const NodeExternal>& ExternalNode)
 	{
 		if (!ExternalNode)
 		{
-			OutResult = FExternalGenerationResult();
-			return;
+			return Tasks::MakeCompletedTask<FExternalGenerationResult>();
 		}
 
 		// See if it was already generated
@@ -32,50 +31,59 @@ namespace UE::Mutable::Private
 			UE::TUniqueLock Lock(GeneratedExternal.Mutex);
 			if (FGeneratedExternalMap::ValueType* Found = GeneratedExternal.Map.Find(Key))
 			{
-				OutResult = *Found;
-				return;
+				return *Found;
 			}
 		}
 		
 		// Generate for each different type of node
-		FMeshTask Result;
+		FExternalTask Result;
+	
 		switch (ExternalNode->GetType()->Type)
 		{
 		case Node::EType::ExternalOperation:
 			{
 				Ptr<const NodeExternalOperation> Operation = static_pointer_cast<const NodeExternalOperation>(ExternalNode);
-				GenerateExternal_Operation(OutResult, Options, Operation);
+				Result = GenerateExternal_Operation(Options, Operation);
 				break;	
 			}
 			
 		case Node::EType::ExternalParameter:
 			{
 				Ptr<const NodeExternalParameter> Parameter = static_pointer_cast<const NodeExternalParameter>(ExternalNode);
-				GenerateExternal_Parameter(OutResult, Options, Parameter);
+				Result = GenerateExternal_Parameter(Options, Parameter);
 				break;
 			}
 			
 		case Node::EType::ExternalSwitch:
 			{
 				Ptr<const NodeExternalSwitch> Parameter = static_pointer_cast<const NodeExternalSwitch>(ExternalNode);
-				GenerateExternal_Switch(OutResult, Options, Parameter);
+				Result = GenerateExternal_Switch(Options, Parameter);
 				break;
 			}
 			
 		default:
+			{
+				Result = Tasks::MakeCompletedTask<FExternalGenerationResult>();
 			unimplemented();
 		}
+		}
 	
 		// Cache the result
 		{
 			UE::TUniqueLock Lock(GeneratedExternal.Mutex);
-			GeneratedExternal.Map.Add(Key, OutResult);
+			GeneratedExternal.Map.Add(Key, Result);
 		}
+		
+		return Result;
 	}
 
 
-	void CodeGenerator::GenerateExternal_Operation(FExternalGenerationResult& OutResult, const FExternalGenerationOptions& Options, const Ptr<const NodeExternalOperation>& ExternalNode)
+	FExternalTask CodeGenerator::GenerateExternal_Operation(const FExternalGenerationOptions& Options, const Ptr<const NodeExternalOperation>& ExternalNode)
 	{
+		return Tasks::Launch(UE_SOURCE_LOCATION, [Options, ExternalNode, this]()
+	{
+			FExternalGenerationResult Result;
+
 		Ptr<ASTOpExternal> External = new ASTOpExternal();
 		External->OperationInstancedStruct = ExternalNode->OperationInstancedStruct;
 
@@ -129,17 +137,17 @@ namespace UE::Mutable::Private
 				WaitTask(Task);
 				FMeshGenerationResult MeshResult = Task.GetResult();
 
-				External->Inputs.Add(ASTChild(External.get(), MeshResult.MeshOp));
+					External->Inputs.Emplace(External.get(), MeshResult.MeshOp);
 				
-				if (!OutResult.MeshGenerationResult.IsSet())
+					if (!Result.MeshGenerationResult.IsSet())
 				{
-					OutResult.MeshGenerationResult = FMeshGenerationResult();
+						Result.MeshGenerationResult = FMeshGenerationResult();
 				}
 				
 				// Merge all generation results from all inputs.
-				OutResult.MeshGenerationResult->GeneratedLayouts.Append(MeshResult.GeneratedLayouts);
-				OutResult.MeshGenerationResult->LayoutOps.Append(MeshResult.LayoutOps);
-				OutResult.MeshGenerationResult->ExtraMeshLayouts.Append(MeshResult.ExtraMeshLayouts);
+					Result.MeshGenerationResult->GeneratedLayouts.Append(MeshResult.GeneratedLayouts);
+					Result.MeshGenerationResult->LayoutOps.Append(MeshResult.LayoutOps);
+					Result.MeshGenerationResult->ExtraMeshLayouts.Append(MeshResult.ExtraMeshLayouts);
 			}
 			else if (Input->GetType()->IsA(NodeImage::GetStaticType()))
 			{
@@ -150,7 +158,7 @@ namespace UE::Mutable::Private
 				FImageGenerationResult InputResult;
 				GenerateImage(ImageOptions, InputResult, static_pointer_cast<NodeImage>(Input));
 
-				External->Inputs.Add(ASTChild(External.get(), InputResult.op));
+					External->Inputs.Emplace(External.get(), InputResult.op);
 			}
 			else if (Input->GetType()->IsA(NodeScalar::GetStaticType()))
 			{
@@ -161,7 +169,7 @@ namespace UE::Mutable::Private
 				FScalarGenerationResult InputResult;
 				GenerateScalar(InputResult, ScalarOptions, static_pointer_cast<NodeScalar>(Input));
 
-				External->Inputs.Add(ASTChild(External.get(), InputResult.op));
+					External->Inputs.Emplace(External.get(), InputResult.op);
 			}
 			else if (Input->GetType()->IsA(NodeColour::GetStaticType()))
 			{
@@ -172,7 +180,7 @@ namespace UE::Mutable::Private
 				FColorGenerationResult InputResult;
 				GenerateColor(InputResult, ColourOptions, static_pointer_cast<NodeColour>(Input));
 
-				External->Inputs.Add(ASTChild(External.get(), InputResult.op));
+					External->Inputs.Emplace(External.get(), InputResult.op);
 			}
 			else if (Input->GetType()->IsA(NodeMaterial::GetStaticType()))
 			{
@@ -183,33 +191,41 @@ namespace UE::Mutable::Private
 				FMaterialGenerationResult InputResult;
 				GenerateMaterial(MaterialOptions, InputResult, static_pointer_cast<NodeMaterial>(Input));
 
-				External->Inputs.Add(ASTChild(External.get(), InputResult.op));
+					External->Inputs.Emplace(External.get(), InputResult.op);
 			}
 			else if (Input->GetType()->IsA(NodeExternal::GetStaticType()))
 			{
-				FExternalGenerationResult InputResult;
-				GenerateExternal(InputResult, Options, static_pointer_cast<NodeExternal>(Input));
+				FExternalTask InputResult = GenerateExternal(Options, static_pointer_cast<NodeExternal>(Input));
+
+				External->Inputs.Emplace(ASTChild::FInvalidTag());
 
-				External->Inputs.Add(ASTChild(External.get(), InputResult.Op));
+				Tasks::AddNested(Tasks::Launch(UE_SOURCE_LOCATION, [InputResult, External, Index]() mutable
+				{
+					External->Inputs[Index] = ASTChild(External.get(), InputResult.GetResult().Op);
+				},
+				InputResult));
 			}
 			else
 			{
-				External->Inputs.Add(ASTChild(External.get()));
+				External->Inputs.Emplace(External.get());
 				unimplemented(); // GenerateMutableSourceExtension has compiled an input type that is not supported.
 			}
 		}
 		
-		if (OutResult.MeshGenerationResult.IsSet())
+		if (Result.MeshGenerationResult.IsSet())
 		{
-			OutResult.MeshGenerationResult->MeshOp = External;	
-			OutResult.MeshGenerationResult->BaseMeshOp = External;			
+				Result.MeshGenerationResult->MeshOp = External;	
+				Result.MeshGenerationResult->BaseMeshOp = External;			
 		}
 		
-		OutResult.Op = External;
+			Result.Op = External;
+			
+			return Result;
+		});
 	}
 
 
-	void CodeGenerator::GenerateExternal_Parameter(FExternalGenerationResult& OutResult, const FExternalGenerationOptions& Options,	const Ptr<const NodeExternalParameter>& ExternalNode)
+	FExternalTask CodeGenerator::GenerateExternal_Parameter(const FExternalGenerationOptions& Options,	const Ptr<const NodeExternalParameter>& ExternalNode)
 	{
 		Ptr<ASTOpParameter> Op;
 
@@ -272,20 +288,28 @@ namespace UE::Mutable::Private
 			}
 		}
 
-		OutResult.Op = Op;
+		FExternalGenerationResult Result;
+		Result.Op = Op;
+		
+		return Tasks::MakeCompletedTask<FExternalGenerationResult>(Result);
 	}
 
 
-	void CodeGenerator::GenerateExternal_Switch(FExternalGenerationResult& OutResult, const FExternalGenerationOptions& Options, const Ptr<const NodeExternalSwitch>& ExternalNode)
+	FExternalTask CodeGenerator::GenerateExternal_Switch(const FExternalGenerationOptions& Options, const Ptr<const NodeExternalSwitch>& ExternalNode)
+	{
+		return Tasks::Launch(UE_SOURCE_LOCATION, [Options, ExternalNode, this]()
 	{
 		if (ExternalNode->Options.Num() == 0)
 		{
+				FExternalGenerationResult Result;
+		
 			// No options in the switch!
 			Ptr<ASTOp> MissingOp = GenerateMissingExternalCode(TEXT("Switch option"),
 				ExternalNode->GetMessageContext());
 			
-			OutResult.Op = MissingOp;
-			return;
+				Result.Op = MissingOp;
+
+				return Result;
 		}
 
 		Ptr<ASTOpSwitch> Op = new ASTOpSwitch();
@@ -304,26 +328,38 @@ namespace UE::Mutable::Private
 			Op->Variable = GenerateMissingScalarCode(TEXT("Switch variable"), 0.0f, ExternalNode->GetMessageContext());
 		}
 
-		// Options
-		for (int32 OptionIndex = 0; OptionIndex < ExternalNode->Options.Num(); ++OptionIndex)
-		{
-			Ptr<ASTOp> Branch;
-			if (ExternalNode->Options[OptionIndex])
-			{
-				FExternalGenerationResult ChildResult;
-				GenerateExternal(ChildResult, Options, ExternalNode->Options[OptionIndex].get());
-				Branch = ChildResult.Op;
-			}
-			else
+			// Options
+			const int32 NumOptions = ExternalNode->Options.Num();
+			Op->Cases.Reserve(NumOptions);
+
+			for (int32 OptionIndex = 0; OptionIndex < ExternalNode->Options.Num(); ++OptionIndex)
 			{
-				// This argument is required
-				Branch = GenerateMissingExternalCode(TEXT("Switch option"), ExternalNode->GetMessageContext());
+				Ptr<ASTOp> Branch;
+				if (ExternalNode->Options[OptionIndex])
+				{
+						FExternalTask ChildResult = GenerateExternal(Options, ExternalNode->Options[OptionIndex].get());
+
+						Op->Cases.Emplace(ASTOpSwitch::FCase::FInvalidTag());
+				
+						Tasks::AddNested(Tasks::Launch(UE_SOURCE_LOCATION, [ChildResult, Op, OptionIndex]() mutable
+						{
+							FExternalGenerationResult Result = ChildResult.GetResult();
+							Op->Cases[OptionIndex] = ASTOpSwitch::FCase(static_cast<int16_t>(OptionIndex), Op, Result.Op);
+						}, 
+						ChildResult));
+				}
+				else
+				{
+					Ptr<ASTOp> Result = GenerateMissingExternalCode(TEXT("Switch option"), ExternalNode->GetMessageContext());
+					Op->Cases.Emplace(static_cast<int16_t>(OptionIndex), Op, Result);
+				}
 			}
 			
-			Op->Cases.Emplace(static_cast<int16_t>(OptionIndex), Op, Branch);
-		}
+			FExternalGenerationResult Result;
+			Result.Op = Op;
 
-		OutResult.Op = Op;
+			return Result;
+		});
 	}
 	
 	
diff --git a/Engine/Plugins/Mutable/Source/MutableTools/Private/MuT/CodeGenerator_FirstPass.h b/Engine/Plugins/Mutable/Source/MutableTools/Private/MuT/CodeGenerator_FirstPass.h
--- a/Engine/Plugins/Mutable/Source/MutableTools/Private/MuT/CodeGenerator_FirstPass.h
+++ b/Engine/Plugins/Mutable/Source/MutableTools/Private/MuT/CodeGenerator_FirstPass.h
@@ -96,6 +96,17 @@ namespace UE::Mutable::Private
 	using FSkeletalMeshTask = UE::Tasks::TTask<FGenericGenerationResult>;
 	using FComponentTask = UE::Tasks::TTask<FGenericGenerationResult>;
 
+	struct FExternalGenerationResult
+	{
+		Ptr<ASTOp> Op;
+    		
+		// If set, at least one of the external operation inputs is a mesh
+		TOptional<FMeshGenerationResult> MeshGenerationResult;
+	};
+	
+	using FExternalTask = UE::Tasks::TTask<FExternalGenerationResult>;
+
+	
 	/** First pass of the code generation process.
 	 * It collects data about the object hierarchy, the conditions for each object and the global modifiers.
 	 */
diff --git a/Engine/Plugins/Mutable/Source/MutableTools/Private/MuT/CodeGenerator_Image.cpp b/Engine/Plugins/Mutable/Source/MutableTools/Private/MuT/CodeGenerator_Image.cpp
--- a/Engine/Plugins/Mutable/Source/MutableTools/Private/MuT/CodeGenerator_Image.cpp
+++ b/Engine/Plugins/Mutable/Source/MutableTools/Private/MuT/CodeGenerator_Image.cpp
@@ -2000,11 +2000,11 @@ namespace UE::Mutable::Private
 	{
 		FExternalGenerationOptions Options;
 		Options.ImageOptions = InOptions;
-					
-		FExternalGenerationResult ExtensionResult;
-		GenerateExternal(ExtensionResult, Options, InNode->Node);
 
-		Result.op = ExtensionResult.Op;
+		FExternalTask ExtensionResult = GenerateExternal(Options, InNode->Node);
+		WaitTask(ExtensionResult);
+
+		Result.op = ExtensionResult.GetResult().Op;
 		Result.Parameter = false;
 		Result.bCanMulticompose = false;
 	}
diff --git a/Engine/Plugins/Mutable/Source/MutableTools/Private/MuT/CodeGenerator_Material.cpp b/Engine/Plugins/Mutable/Source/MutableTools/Private/MuT/CodeGenerator_Material.cpp
--- a/Engine/Plugins/Mutable/Source/MutableTools/Private/MuT/CodeGenerator_Material.cpp
+++ b/Engine/Plugins/Mutable/Source/MutableTools/Private/MuT/CodeGenerator_Material.cpp
@@ -397,10 +397,10 @@ namespace UE::Mutable::Private
 		FExternalGenerationOptions Options;
 		Options.MaterialOptions = InOptions;
 					
-		FExternalGenerationResult ExtensionResult;
-		GenerateExternal(ExtensionResult, Options, InNode->Node);
+		FExternalTask ExtensionResult = GenerateExternal(Options, InNode->Node);
+		WaitTask(ExtensionResult);
 
-		Result.op = ExtensionResult.Op;
+		Result.op = ExtensionResult.GetResult().Op;
 	}
 
 	
diff --git a/Engine/Plugins/Mutable/Source/MutableTools/Private/MuT/CodeGenerator_Mesh.cpp b/Engine/Plugins/Mutable/Source/MutableTools/Private/MuT/CodeGenerator_Mesh.cpp
--- a/Engine/Plugins/Mutable/Source/MutableTools/Private/MuT/CodeGenerator_Mesh.cpp
+++ b/Engine/Plugins/Mutable/Source/MutableTools/Private/MuT/CodeGenerator_Mesh.cpp
@@ -1704,27 +1704,27 @@ namespace UE::Mutable::Private
 
 	FMeshTask CodeGenerator::GenerateMesh_External(const FMeshGenerationStaticOptions& StaticOptions, FMeshOptionsTask InOptions, const NodeMeshExternal* Node)
 	{
-		return UE::Tasks::Launch(TEXT("MutableExtension"),
-			[Node, StaticOptions, InOptions, this]() mutable
-			{
-				FExternalGenerationOptions Options;
-				Options.MeshStaticOptions = StaticOptions;
-				Options.MeshDynamicOptions = InOptions;
-					
-				FExternalGenerationResult ExtensionResult;
-				GenerateExternal(ExtensionResult, Options, Node->Node);
+		FExternalGenerationOptions Options;
+		Options.MeshStaticOptions = StaticOptions;
+		Options.MeshDynamicOptions = InOptions;
 
-				FMeshGenerationResult MeshResult;
+		FExternalTask ExternalTask = GenerateExternal(Options, Node->Node);
 
-				if (ExtensionResult.MeshGenerationResult.IsSet())
-				{
-					MeshResult = ExtensionResult.MeshGenerationResult.GetValue();
-				}
-					
-				MeshResult.MeshOp = ExtensionResult.Op;
-					
-				return MeshResult;
+		return UE::Tasks::Launch(UE_SOURCE_LOCATION,[ExternalTask, this]() mutable
+		{
+			const FExternalGenerationResult& ExtensionResult = ExternalTask.GetResult();
+				
+			FMeshGenerationResult MeshResult;
+
+			if (ExtensionResult.MeshGenerationResult.IsSet())
+			{
+				MeshResult = ExtensionResult.MeshGenerationResult.GetValue();
 			}
-		);
+				
+			MeshResult.MeshOp = ExtensionResult.Op;
+				
+			return MeshResult;
+		},
+		ExternalTask);
 	}
 }
diff --git a/Engine/Plugins/Mutable/Source/MutableTools/Private/MuT/CodeGenerator_Scalar.cpp b/Engine/Plugins/Mutable/Source/MutableTools/Private/MuT/CodeGenerator_Scalar.cpp
--- a/Engine/Plugins/Mutable/Source/MutableTools/Private/MuT/CodeGenerator_Scalar.cpp
+++ b/Engine/Plugins/Mutable/Source/MutableTools/Private/MuT/CodeGenerator_Scalar.cpp
@@ -459,10 +459,10 @@ namespace UE::Mutable::Private
 		FExternalGenerationOptions Options;
 		Options.ScalarOptions = InOptions;
 					
-		FExternalGenerationResult ExtensionResult;
-		GenerateExternal(ExtensionResult, Options, InNode->Node);
+		FExternalTask ExtensionResult = GenerateExternal(Options, InNode->Node);
+		WaitTask(ExtensionResult);
 
-		Result.op = ExtensionResult.Op;
+		Result.op = ExtensionResult.GetResult().Op;
 	}
 
 
