[mutable] Made TextureLayer node unable to accept layer connections of Color type if not explicitly set by the user.
- This will prevent connections of types different from the one visually shown by the node "Layer" pin.

TextureLayer Node
- Now the "Layer" pin can change type if the user so desires.
- Removed code designed to keep data between node copies but with no real use.
- Added backwards compatibility code for the changes described above.

TextureFromColor Node
- Updated the structure of the TextureFromColor node to hold the pin in local reference variables.
- Added backwards compatibility code for the changes in TextureFromColor too.

PinTypeSelector
- Added tooltip text for the pin type selection slate.

#Synced-CL 52244883

diff --git a/Engine/Plugins/Mutable/Source/CustomizableObject/Internal/MuCO/CustomizableObjectCustomVersion.h b/Engine/Plugins/Mutable/Source/CustomizableObject/Internal/MuCO/CustomizableObjectCustomVersion.h
--- a/Engine/Plugins/Mutable/Source/CustomizableObject/Internal/MuCO/CustomizableObjectCustomVersion.h
+++ b/Engine/Plugins/Mutable/Source/CustomizableObject/Internal/MuCO/CustomizableObjectCustomVersion.h
@@ -233,6 +233,12 @@ struct FCustomizableObjectCustomVersion
 		RemovedTypedSwitches,
 		
 		SkeletalMeshMakeNodeDeprecations,				// SKM Component has gotten some responsibilities stripped from it and moved to the new Skeletal Mesh Make node
+
+		TextureFromColorPinCaching,
+		
+		TextureLayerNodePinCaching,
+		
+		TextureLayerNodeExplicitLayerTypePins,
 		
 		// -----<new versions can be added above this line>-------------------------------------------------
 		VersionPlusOne,
diff --git a/Engine/Plugins/Mutable/Source/CustomizableObjectEditor/Private/MuCOE/GenerateMutableSource/GenerateMutableSource.cpp b/Engine/Plugins/Mutable/Source/CustomizableObjectEditor/Private/MuCOE/GenerateMutableSource/GenerateMutableSource.cpp
--- a/Engine/Plugins/Mutable/Source/CustomizableObjectEditor/Private/MuCOE/GenerateMutableSource/GenerateMutableSource.cpp
+++ b/Engine/Plugins/Mutable/Source/CustomizableObjectEditor/Private/MuCOE/GenerateMutableSource/GenerateMutableSource.cpp
@@ -483,14 +483,14 @@ UTexture2D* FindReferenceImage(const UEdGraphPin* Pin, FMutableGraphGenerationCo
 
 	else if (const UCustomizableObjectNodeTextureLayer* TypedNodeLayer = Cast<UCustomizableObjectNodeTextureLayer>(Node))
 	{
-		if (const UEdGraphPin* ConnectedPin = FollowInputPin(*TypedNodeLayer->BasePin()))
+		if (const UEdGraphPin* ConnectedPin = FollowInputPin(*TypedNodeLayer->BasePin.Get()))
 		{
 			Result = FindReferenceImage(ConnectedPin, GenerationContext);
 		}
 
 		for (int LayerIndex = 0; !Result && LayerIndex < TypedNodeLayer->GetNumLayers(); ++LayerIndex)
 		{
-			if (const UEdGraphPin* ConnectedPin = FollowInputPin(*TypedNodeLayer->LayerPin(LayerIndex)))
+			if (const UEdGraphPin* ConnectedPin = FollowInputPin(*TypedNodeLayer->GetLayerPin(LayerIndex)))
 			{
 				if (ConnectedPin->PinType.PinCategory == Schema->PC_Texture)
 				{
diff --git a/Engine/Plugins/Mutable/Source/CustomizableObjectEditor/Private/MuCOE/GenerateMutableSource/GenerateMutableSourceImage.cpp b/Engine/Plugins/Mutable/Source/CustomizableObjectEditor/Private/MuCOE/GenerateMutableSource/GenerateMutableSourceImage.cpp
--- a/Engine/Plugins/Mutable/Source/CustomizableObjectEditor/Private/MuCOE/GenerateMutableSource/GenerateMutableSourceImage.cpp
+++ b/Engine/Plugins/Mutable/Source/CustomizableObjectEditor/Private/MuCOE/GenerateMutableSource/GenerateMutableSourceImage.cpp
@@ -277,17 +277,17 @@ UE::Mutable::Private::Ptr<UE::Mutable::Private::NodeImage> GenerateMutableSource
 
 	else if (const UCustomizableObjectNodeTextureLayer* TypedNodeLayer = Cast<UCustomizableObjectNodeTextureLayer>(Node))
 	{
-		if (const UEdGraphPin* ConnectedPin = FollowInputPin(*TypedNodeLayer->BasePin()))
+		if (const UEdGraphPin* ConnectedPin = FollowInputPin(*TypedNodeLayer->BasePin.Get()))
 		{
 			Result = GenerateMutableSourceImage(ConnectedPin, GenerationContext, ReferenceTextureSize);
 		}
 
 		for (int LayerIndex = 0; LayerIndex < TypedNodeLayer->GetNumLayers(); ++LayerIndex)
 		{
-			if (const UEdGraphPin* OtherPin = FollowInputPin(*TypedNodeLayer->LayerPin(LayerIndex)))
+			if (const UEdGraphPin* OtherPin = FollowInputPin(*TypedNodeLayer->GetLayerPin(LayerIndex)))
 			{
 				UE::Mutable::Private::Ptr<UE::Mutable::Private::NodeImage> MaskNode = nullptr;
-				if (const UEdGraphPin* ConnectedPin = FollowInputPin(*TypedNodeLayer->MaskPin(LayerIndex)))
+				if (const UEdGraphPin* ConnectedPin = FollowInputPin(*TypedNodeLayer->GetMaskPin(LayerIndex)))
 				{
 					MaskNode = GenerateMutableSourceImage(ConnectedPin, GenerationContext, ReferenceTextureSize);
 				}
@@ -326,7 +326,6 @@ UE::Mutable::Private::Ptr<UE::Mutable::Private::NodeImage> GenerateMutableSource
 					LayerNode->Mask = MaskNode;
 					Result = LayerNode;
 				}
-
 				else if (OtherPin->PinType.PinCategory == Schema->PC_Color)
 				{
 					UE::Mutable::Private::Ptr<UE::Mutable::Private::NodeColour> ColorNode = GenerateMutableSourceColor(OtherPin, GenerationContext);
@@ -337,6 +336,10 @@ UE::Mutable::Private::Ptr<UE::Mutable::Private::NodeImage> GenerateMutableSource
 					LayerNode->Colour = ColorNode;
 					LayerNode->Mask = MaskNode;
 					Result = LayerNode;
+				}			
+				else
+				{	
+					unimplemented();
 				}
 
 				// We need it here because we create multiple nodes.
@@ -458,7 +461,7 @@ UE::Mutable::Private::Ptr<UE::Mutable::Private::NodeImage> GenerateMutableSource
 	{
 		UE::Mutable::Private::Ptr<UE::Mutable::Private::NodeColour> color;
 
-		if (const UEdGraphPin* ConnectedPin = FollowInputPin(*TypedNodeFromColor->ColorPin()))
+		if (const UEdGraphPin* ConnectedPin = FollowInputPin(*TypedNodeFromColor->ColorPin.Get()))
 		{
 			color = GenerateMutableSourceColor(ConnectedPin, GenerationContext);
 		}
diff --git a/Engine/Plugins/Mutable/Source/CustomizableObjectEditor/Private/MuCOE/Nodes/CONodeCopySkeletalMeshSection.cpp b/Engine/Plugins/Mutable/Source/CustomizableObjectEditor/Private/MuCOE/Nodes/CONodeCopySkeletalMeshSection.cpp
--- a/Engine/Plugins/Mutable/Source/CustomizableObjectEditor/Private/MuCOE/Nodes/CONodeCopySkeletalMeshSection.cpp
+++ b/Engine/Plugins/Mutable/Source/CustomizableObjectEditor/Private/MuCOE/Nodes/CONodeCopySkeletalMeshSection.cpp
@@ -192,16 +192,16 @@ UCONodeSkeletalMeshSection* UCONodeCopySkeletalMeshSection::GetMaterialNode() co
 
 bool UCONodeCopySkeletalMeshSection::CanConnect(const UEdGraphPin* InOwnedInputPin, const UEdGraphPin* InOutputPin, bool& bOutIsOtherNodeBlocklisted, bool& bOutArePinsCompatible) const
 {
-	if (!Super::CanConnect(InOwnedInputPin, InOutputPin, bOutIsOtherNodeBlocklisted, bOutIsOtherNodeBlocklisted))
+	if (!Super::CanConnect(InOwnedInputPin, InOutputPin, bOutIsOtherNodeBlocklisted, bOutArePinsCompatible))
 	{
 		return false;
 	}
 
 	if (InOwnedInputPin == GetMeshSectionPin())
 	{
-		const UEdGraphNode* OuputPinOwningNode = InOutputPin->GetOwningNode();
-		return (OuputPinOwningNode->IsA(UCONodeSkeletalMeshSection::StaticClass()) && !OuputPinOwningNode->IsA(UCONodeCopySkeletalMeshSection::StaticClass()))
-			|| OuputPinOwningNode->IsA(UCustomizableObjectNodeExternalPin::StaticClass());
+		const UEdGraphNode* OutputPinOwningNode = InOutputPin->GetOwningNode();
+		return (OutputPinOwningNode->IsA(UCONodeSkeletalMeshSection::StaticClass()) && !OutputPinOwningNode->IsA(UCONodeCopySkeletalMeshSection::StaticClass()))
+			|| OutputPinOwningNode->IsA(UCustomizableObjectNodeExternalPin::StaticClass());
 	}
 
 	return true;
diff --git a/Engine/Plugins/Mutable/Source/CustomizableObjectEditor/Private/MuCOE/Nodes/CustomizableObjectNodeTextureFromColor.cpp b/Engine/Plugins/Mutable/Source/CustomizableObjectEditor/Private/MuCOE/Nodes/CustomizableObjectNodeTextureFromColor.cpp
--- a/Engine/Plugins/Mutable/Source/CustomizableObjectEditor/Private/MuCOE/Nodes/CustomizableObjectNodeTextureFromColor.cpp
+++ b/Engine/Plugins/Mutable/Source/CustomizableObjectEditor/Private/MuCOE/Nodes/CustomizableObjectNodeTextureFromColor.cpp
@@ -14,24 +14,31 @@ class UCustomizableObjectNodeRemapPins;
 
 void UCustomizableObjectNodeTextureFromColor::AllocateDefaultPins(UCustomizableObjectNodeRemapPins* RemapPins)
 {
-	CustomCreatePin(EGPD_Output, UEdGraphSchema_CustomizableObject::PC_Texture);
-	CustomCreatePin(EGPD_Input, UEdGraphSchema_CustomizableObject::PC_Color);
+	TexturePin = CustomCreatePin(EGPD_Output, UEdGraphSchema_CustomizableObject::PC_Texture);
+	ColorPin = CustomCreatePin(EGPD_Input, UEdGraphSchema_CustomizableObject::PC_Color);
 }
 
 
-
 void UCustomizableObjectNodeTextureFromColor::BackwardsCompatibleFixup(int32 CustomizableObjectCustomVersion)
 {
 	Super::BackwardsCompatibleFixup(CustomizableObjectCustomVersion);
 	
 	if (CustomizableObjectCustomVersion == FCustomizableObjectCustomVersion::FixPinsNamesImageToTexture2)
 	{
-		if (UEdGraphPin* TexturePin = FindPin(TEXT("Image")))
+		if (UEdGraphPin* FoundTexturePin = FindPin(TEXT("Image")))
 		{
-			TexturePin->PinName = TEXT("Texture");
+			FoundTexturePin->PinName = TEXT("Texture");
 			UCustomizableObjectNode::ReconstructNode();
 		}
 	}
+	else if (CustomizableObjectCustomVersion == FCustomizableObjectCustomVersion::TextureFromColorPinCaching)
+	{
+		ColorPin = FindPin(TEXT("Color"));
+		check(ColorPin.Get());
+		
+		TexturePin = FindPin(TEXT("Texture"));
+		check(TexturePin.Get());
+	}
 }
 
 
diff --git a/Engine/Plugins/Mutable/Source/CustomizableObjectEditor/Private/MuCOE/Nodes/CustomizableObjectNodeTextureFromColor.h b/Engine/Plugins/Mutable/Source/CustomizableObjectEditor/Private/MuCOE/Nodes/CustomizableObjectNodeTextureFromColor.h
--- a/Engine/Plugins/Mutable/Source/CustomizableObjectEditor/Private/MuCOE/Nodes/CustomizableObjectNodeTextureFromColor.h
+++ b/Engine/Plugins/Mutable/Source/CustomizableObjectEditor/Private/MuCOE/Nodes/CustomizableObjectNodeTextureFromColor.h
@@ -6,8 +6,6 @@
 
 #include "CustomizableObjectNodeTextureFromColor.generated.h"
 
-#define UE_API CUSTOMIZABLEOBJECTEDITOR_API
-
 namespace ENodeTitleType { enum Type : int; }
 
 class UCustomizableObjectNodeRemapPins;
@@ -15,25 +13,24 @@ class UEdGraphPin;
 class UObject;
 
 
-UCLASS(MinimalAPI)
+UCLASS()
 class UCustomizableObjectNodeTextureFromColor : public UCustomizableObjectNode
 {
 public:
 	GENERATED_BODY()
 
 	// Begin EdGraphNode interface
-	UE_API FText GetNodeTitle(ENodeTitleType::Type TitleType) const override;
-	UE_API FLinearColor GetNodeTitleColor() const override;
-	UE_API FText GetTooltipText() const override;
+	virtual FText GetNodeTitle(ENodeTitleType::Type TitleType) const override;
+	virtual FLinearColor GetNodeTitleColor() const override;
+	virtual FText GetTooltipText() const override;
 
 	// UCustomizableObjectNode interface
-	UE_API virtual void AllocateDefaultPins(UCustomizableObjectNodeRemapPins* RemapPins) override;
-	UE_API virtual void BackwardsCompatibleFixup(int32 CustomizableObjectCustomVersion) override;
+	virtual void AllocateDefaultPins(UCustomizableObjectNodeRemapPins* RemapPins) override;
+	virtual void BackwardsCompatibleFixup(int32 CustomizableObjectCustomVersion) override;
 
-	UEdGraphPin* ColorPin() const
-	{
-		return FindPin(TEXT("Color"));
-	}
-};
+	UPROPERTY()
+	FEdGraphPinReference ColorPin;
 
-#undef UE_API
+	UPROPERTY()
+	FEdGraphPinReference TexturePin;
+};
diff --git a/Engine/Plugins/Mutable/Source/CustomizableObjectEditor/Private/MuCOE/Nodes/CustomizableObjectNodeTextureLayer.cpp b/Engine/Plugins/Mutable/Source/CustomizableObjectEditor/Private/MuCOE/Nodes/CustomizableObjectNodeTextureLayer.cpp
--- a/Engine/Plugins/Mutable/Source/CustomizableObjectEditor/Private/MuCOE/Nodes/CustomizableObjectNodeTextureLayer.cpp
+++ b/Engine/Plugins/Mutable/Source/CustomizableObjectEditor/Private/MuCOE/Nodes/CustomizableObjectNodeTextureLayer.cpp
@@ -2,7 +2,9 @@
 
 #include "MuCOE/Nodes/CustomizableObjectNodeTextureLayer.h"
 
+#include "CustomizableObjectNodeTextureFromColor.h"
 #include "MuCO/CustomizableObjectCustomVersion.h"
+#include "MuCOE/CustomizableObjectSchemaActions.h"
 #include "MuCOE/EdGraphSchema_CustomizableObject.h"
 #include "MuCOE/UnrealEditorPortabilityHelpers.h"
 
@@ -10,10 +12,10 @@
 
 class UCustomizableObjectNodeRemapPins;
 class UEdGraph;
+class UEdGraphNode;
 
 #define LOCTEXT_NAMESPACE "CustomizableObjectEditor"
 
-TArray<UCustomizableObjectNodeTextureLayer::TextureLayerLayersData> UCustomizableObjectNodeTextureLayer::ArrayDestroyedNodes;
 
 UCustomizableObjectNodeTextureLayer::UCustomizableObjectNodeTextureLayer()
 	: Super()
@@ -25,7 +27,8 @@ UCustomizableObjectNodeTextureLayer::UCustomizableObjectNodeTextureLayer()
 void UCustomizableObjectNodeTextureLayer::PostEditChangeProperty(FPropertyChangedEvent& PropertyChangedEvent)
 {
 	FProperty* PropertyThatChanged = PropertyChangedEvent.Property;
-	if ( PropertyThatChanged && PropertyThatChanged->GetName() == TEXT("Layers") )
+	if (PropertyThatChanged && 
+		PropertyThatChanged->GetName() == GET_MEMBER_NAME_CHECKED(UCustomizableObjectNodeTextureLayer, Layers))
 	{
 		ReconstructNode();
 	}
@@ -49,19 +52,31 @@ void UCustomizableObjectNodeTextureLayer::Serialize(FArchive& Ar)
 
 void UCustomizableObjectNodeTextureLayer::AllocateDefaultPins(UCustomizableObjectNodeRemapPins* RemapPins)
 {
-	const UEdGraphSchema_CustomizableObject* Schema = GetDefault<UEdGraphSchema_CustomizableObject>();
+	OutputPinReference = CustomCreatePin(EGPD_Output, UEdGraphSchema_CustomizableObject::PC_Texture, TEXT("Texture"), LOCTEXT("Texture", "Texture"));
 
-	OutputPinReference =CustomCreatePin(EGPD_Output, Schema->PC_Texture, TEXT("Texture"), LOCTEXT("Texture", "Texture"));
+	BasePin = CustomCreatePin(EGPD_Input, UEdGraphSchema_CustomizableObject::PC_Texture, TEXT("Base"), LOCTEXT("Base", "Base"));
 	
-	CustomCreatePin(EGPD_Input, Schema->PC_Texture, TEXT("Base"), LOCTEXT("Base", "Base"));
-	
-	for (int32 LayerIndex = 0; LayerIndex < Layers.Num() ; ++LayerIndex )
+	for (int32 LayerIndex = 0; LayerIndex < Layers.Num(); ++LayerIndex)
 	{
-		FString PinName = FString::Printf(TEXT("Layer %d "), LayerIndex);
-		CustomCreatePin(EGPD_Input, Schema->PC_Texture, *PinName, FText::FromString(PinName));
+		FCustomizableObjectTextureLayer& CurrentLayerData = Layers[LayerIndex];
+	
+		// Determine the type of node the new layer pin will be. We use the old type so we keep the type of operation and tye pin type in sync
+		{
+			FName NewLayerPinType = UEdGraphSchema_CustomizableObject::PC_Texture;
+
+			// If possible use the same pin type as the pin this Layer was representing before resizing the array
+			const FEdGraphPinReference LayerPinReference = CurrentLayerData.LayerPin;
+			if (const UEdGraphPin* LayerPin = LayerPinReference.Get())
+			{
+				NewLayerPinType = LayerPin->PinType.PinCategory;
+			}
+		
+			const FText LayerPinName = FText::Format(LOCTEXT("Texture_Layer_LayerPinName", "Layer {0}"), LayerIndex);
+			CurrentLayerData.LayerPin = CustomCreatePin(EGPD_Input, NewLayerPinType, *LayerPinName.ToString(), LayerPinName);
+		}
 
-		PinName = FString::Printf( TEXT("Mask %d "), LayerIndex );
-		CustomCreatePin(EGPD_Input, Schema->PC_Texture, *PinName, FText::FromString(PinName));
+		const FText MaskPinName = FText::Format(LOCTEXT("Texture_Layer_MaskPinName", "Mask {0}"), LayerIndex);
+		CurrentLayerData.MaskPin = CustomCreatePin(EGPD_Input, UEdGraphSchema_CustomizableObject::PC_Texture, *MaskPinName.ToString(), MaskPinName);
 	}
 }
 
@@ -78,91 +93,95 @@ void UCustomizableObjectNodeTextureLayer::BackwardsCompatibleFixup(int32 Customi
 			UCustomizableObjectNode::ReconstructNode();
 		}
 	}
-}
-
-
-FText UCustomizableObjectNodeTextureLayer::GetNodeTitle(ENodeTitleType::Type TitleType) const
-{
-	return LOCTEXT("Texture_Layer", "Texture Layer");
-}
 
+	else if (CustomizableObjectCustomVersion == FCustomizableObjectCustomVersion::TextureLayerNodePinCaching)
+	{
+		// Cache base input pin
+		BasePin = FindPin(TEXT("Base"), EEdGraphPinDirection::EGPD_Input);
+		check(BasePin.Get());
 
-FLinearColor UCustomizableObjectNodeTextureLayer::GetNodeTitleColor() const
-{
-	return UEdGraphSchema_CustomizableObject::GetPinTypeColor(UEdGraphSchema_CustomizableObject::PC_Texture);
-}
-
+		// Cache layer pin data
+		for (int32 LayerIndex = 0; LayerIndex < Layers.Num(); ++LayerIndex)
+		{
+			FCustomizableObjectTextureLayer& CurrentLayer = Layers[LayerIndex];
 
-void UCustomizableObjectNodeTextureLayer::PrepareForCopying()
-{
-	// Destroy previous copy with possible out of date information just in case
-	int32 IndexToRemove = GetIndexInArrayDestroyedNodes(this);
+			const FString LayerPinName = FString::Printf(TEXT("Layer %d "), LayerIndex);
+			CurrentLayer.LayerPin = FindPin(LayerPinName, EGPD_Input);
+			check(CurrentLayer.LayerPin.Get());
 
-	if (IndexToRemove != -1)
-	{
-		ArrayDestroyedNodes.RemoveAt(IndexToRemove);
+			const FString MaskPinName = FString::Printf(TEXT("Mask %d "), LayerIndex);
+			CurrentLayer.MaskPin = FindPin(MaskPinName, EGPD_Input);
+			check(CurrentLayer.MaskPin.Get());
+		}
 	}
 
-	TextureLayerLayersData Data;
-	Data.Name = GetName();
-	Data.Graph = GetGraph();
-	Data.Layers = Layers;
-	ArrayDestroyedNodes.Add(Data);
-}
+	else if (CustomizableObjectCustomVersion == FCustomizableObjectCustomVersion::TextureLayerNodeExplicitLayerTypePins)
+	{
+		bool bNodeWasModified = false;
 
-void UCustomizableObjectNodeTextureLayer::PostPasteNode()
-{
-	int32 IndexToRemove = GetIndexInArrayDestroyedNodes(this);
+		// Iterate over all Layer pins as those are the ones that could have a color connected. Then create a new TextureFromColor node and connect it
+		for (int32 LayerIndex = 0; LayerIndex < Layers.Num(); ++LayerIndex)
+		{
+			FCustomizableObjectTextureLayer& CurrentLayer = Layers[LayerIndex];
+
+			UEdGraphPin* LayerPin = CurrentLayer.LayerPin.Get();
+			check(LayerPin);
+			check(LayerPin->PinType.PinCategory == UEdGraphSchema_CustomizableObject::PC_Texture);
+				
+			// Who is connected to this Layer Pin?
+			if (LayerPin->LinkedTo.IsEmpty())
+			{
+				continue;
+			}
+
+			const UEdGraphPin* ConnectedNodePin = LayerPin->LinkedTo[0];
+			LayerPin->PinType = ConnectedNodePin->PinType;
+
+			bNodeWasModified = true;
+		}
 
-	if (IndexToRemove != -1)
-	{
-		Layers = ArrayDestroyedNodes[IndexToRemove].Layers;
-		ArrayDestroyedNodes.RemoveAt(IndexToRemove);
+		if (bNodeWasModified)
+		{
+			UEdGraph* Graph = GetGraph();
+			check(Graph)
+			Graph->NotifyNodeChanged(this);
+		}
 	}
 }
 
 
-bool UCustomizableObjectNodeTextureLayer::CanConnect(const UEdGraphPin* InOwnedInputPin, const UEdGraphPin* InOutputPin, bool& bOutIsOtherNodeBlocklisted, bool& bOutArePinsCompatible) const
+TArray<FName> UCustomizableObjectNodeTextureLayer::GetPinAllowedTypes(const UEdGraphPin& Pin) const
 {
-	const UEdGraphSchema_CustomizableObject* Schema = GetDefault<UEdGraphSchema_CustomizableObject>();
+	// If this is a layer pin then expose the texture and color options
+	const int32 LayerCount = Layers.Num();
+	for (int32 LayerIndex = 0; LayerIndex < LayerCount; ++LayerIndex)
+	{
+		const UEdGraphPin* LayerPin = this->GetLayerPin(LayerIndex);
+		check(LayerPin);
 
-	// In layer pins we accept both colors and images
-	for ( int32 LayerIndex = 0; LayerIndex < Layers.Num(); ++LayerIndex )
+		const FEdGraphPinReference LayerRef = LayerPin;
+		if (LayerRef == FEdGraphPinReference(&Pin))
 	{
-		if (LayerPin(LayerIndex) == InOwnedInputPin )
-		{
-			return (InOutputPin->PinType.PinCategory==Schema->PC_Texture)
-				|| (InOutputPin->PinType.PinCategory==Schema->PC_Color);
+			return { 
+				UEdGraphSchema_CustomizableObject::PC_Texture, 
+				UEdGraphSchema_CustomizableObject::PC_Color, 
+			};
 		}
 	}
 
-	return Super::CanConnect( InOwnedInputPin, InOutputPin, bOutIsOtherNodeBlocklisted, bOutArePinsCompatible);
+	return Super::GetPinAllowedTypes(Pin);
 }
 
 
-UEdGraphPin* UCustomizableObjectNodeTextureLayer::OutputPin() const
+FText UCustomizableObjectNodeTextureLayer::GetNodeTitle(ENodeTitleType::Type TitleType) const
 {
-	return OutputPinReference.Get();
+	return LOCTEXT("Texture_Layer", "Texture Layer");
 }
 
 
-int32 UCustomizableObjectNodeTextureLayer::GetIndexInArrayDestroyedNodes(UCustomizableObjectNodeTextureLayer* Node)
+FLinearColor UCustomizableObjectNodeTextureLayer::GetNodeTitleColor() const
 {
-	const int32 MaxIndex = ArrayDestroyedNodes.Num();
-	FString NodeName = Node->GetName();
-	UEdGraph* NodeGraph = Node->GetGraph();
-	for (int32 i = 0; i < MaxIndex; ++i)
-	{
-		TextureLayerLayersData& DataTemp = ArrayDestroyedNodes[i];
-
-		//if ((DataTemp.Name == NodeName) && (DataTemp.Graph == NodeGraph))
-		if (DataTemp.Name == NodeName) // Copying between graphs would yeld incorredct array data, are the names of the elements unique?
-		{
-			return i;
-		}
-	}
-
-	return -1;
+	return UEdGraphSchema_CustomizableObject::GetPinTypeColor(UEdGraphSchema_CustomizableObject::PC_Texture);
 }
 
 
diff --git a/Engine/Plugins/Mutable/Source/CustomizableObjectEditor/Private/MuCOE/Nodes/CustomizableObjectNodeTextureLayer.h b/Engine/Plugins/Mutable/Source/CustomizableObjectEditor/Private/MuCOE/Nodes/CustomizableObjectNodeTextureLayer.h
--- a/Engine/Plugins/Mutable/Source/CustomizableObjectEditor/Private/MuCOE/Nodes/CustomizableObjectNodeTextureLayer.h
+++ b/Engine/Plugins/Mutable/Source/CustomizableObjectEditor/Private/MuCOE/Nodes/CustomizableObjectNodeTextureLayer.h
@@ -6,8 +6,6 @@
 
 #include "CustomizableObjectNodeTextureLayer.generated.h"
 
-#define UE_API CUSTOMIZABLEOBJECTEDITOR_API
-
 namespace ENodeTitleType { enum Type : int; }
 
 class FArchive;
@@ -44,64 +42,68 @@ struct FCustomizableObjectTextureLayer
 	}
 
 	UPROPERTY(EditAnywhere, Category=CustomizableObject)
-	TEnumAsByte<enum ECustomizableObjectTextureLayerEffect> Effect;
+	TEnumAsByte<ECustomizableObjectTextureLayerEffect> Effect;
+	
+	UPROPERTY()
+	FEdGraphPinReference LayerPin;
+	
+	UPROPERTY()
+	FEdGraphPinReference MaskPin;
 };
 
 
-UCLASS(MinimalAPI)
+UCLASS()
 class UCustomizableObjectNodeTextureLayer : public UCustomizableObjectNode
 {
 public:
 	GENERATED_BODY()
 
-	UE_API UCustomizableObjectNodeTextureLayer();
-
-	/**  */
-	UPROPERTY(EditAnywhere, Category=CustomizableObject)
-	TArray<FCustomizableObjectTextureLayer> Layers;
+	UCustomizableObjectNodeTextureLayer();
 
 	// UObject interface.
-	UE_API void PostEditChangeProperty(FPropertyChangedEvent& PropertyChangedEvent) override;
-	UE_API virtual void Serialize(FArchive& Ar) override;
+	virtual void PostEditChangeProperty(FPropertyChangedEvent& PropertyChangedEvent) override;
+	virtual void Serialize(FArchive& Ar) override;
 	
 	// Begin EdGraphNode interface
-	UE_API FText GetNodeTitle(ENodeTitleType::Type TitleType) const override;
-	UE_API FLinearColor GetNodeTitleColor() const override;
-	UE_API FText GetTooltipText() const override;
-	UE_API void PrepareForCopying() override;
-	UE_API void PostPasteNode() override;
+	virtual FText GetNodeTitle(ENodeTitleType::Type TitleType) const override;
+	virtual FLinearColor GetNodeTitleColor() const override;
+	virtual FText GetTooltipText() const override;
 
 	// UCustomizableObjectNode interface
-	UE_API virtual void AllocateDefaultPins(UCustomizableObjectNodeRemapPins* RemapPins) override;
-	UE_API virtual void BackwardsCompatibleFixup(int32 CustomizableObjectCustomVersion) override;
-	UE_API virtual bool CanConnect(const UEdGraphPin* InOwnedInputPin, const UEdGraphPin* InOutputPin, bool& bOutIsOtherNodeBlocklisted, bool& bOutArePinsCompatible) const override;
+	virtual void AllocateDefaultPins(UCustomizableObjectNodeRemapPins* RemapPins) override;
+	virtual void BackwardsCompatibleFixup(int32 CustomizableObjectCustomVersion) override;
+	virtual TArray<FName> GetPinAllowedTypes(const UEdGraphPin& Pin) const override;
 		
-	UEdGraphPin* BasePin() const
+	UEdGraphPin* GetLayerPin(int32 Index) const
 	{
-		return FindPin(TEXT("Base"));
-	}
-
-	UE_API UEdGraphPin* OutputPin() const;
+		if (Layers.IsValidIndex(Index))
+		{
+			return Layers[Index].LayerPin.Get();
+		}
 
-	UEdGraphPin* LayerPin(int32 Index) const
-	{
-		FString PinName = FString::Printf(TEXT("Layer %d "), Index);
-		return FindPin(PinName);
+		return nullptr;
 	}
 
-	UEdGraphPin* MaskPin(int32 Index) const
+	UEdGraphPin* GetMaskPin(int32 Index) const
 	{
-		FString PinName = FString::Printf(TEXT("Mask %d "), Index);
-		return FindPin(PinName);
+		if (Layers.IsValidIndex(Index))
+		{
+			return Layers[Index].MaskPin.Get();
+		}
+
+		return nullptr;
 	}
 
+	
 	int32 GetNumLayers() const
 	{
 		int32 Count = 0;
 
-		for (UEdGraphPin* Pin : GetAllNonOrphanPins())
+		for (const FCustomizableObjectTextureLayer& LayerObject : Layers)
 		{
-			if (Pin->GetName().StartsWith(TEXT("Layer ")))
+			const UEdGraphPin* LayerPin = LayerObject.LayerPin.Get();
+			check(LayerPin);
+			if (!LayerPin->bOrphanedPin)
 			{
 				Count++;
 			}
@@ -109,21 +111,16 @@ public:
 
 		return Count;
 	}
-protected:
-	// This struct will store the information of each textue layer node when its destroyed, to be able to recover the UCustomizableObjectNodeTextureLayer::Layers
-	// data in a post paste operation. Otherwise, this information is lost.
-	struct TextureLayerLayersData
-	{
-		FString Name; // Texture layer node name
-		TArray<FCustomizableObjectTextureLayer> Layers;
-		UEdGraph* Graph; // Graph that contained this node
-	};
-	static UE_API int32 GetIndexInArrayDestroyedNodes(UCustomizableObjectNodeTextureLayer* Node);
-	static UE_API TArray<TextureLayerLayersData> ArrayDestroyedNodes;
+	
+	UPROPERTY()
+	FEdGraphPinReference BasePin;
+	
+	/**  */
+	UPROPERTY(EditAnywhere, Category=CustomizableObject)
+	TArray<FCustomizableObjectTextureLayer> Layers;
 
 private:
 	UPROPERTY()
 	FEdGraphPinReference OutputPinReference;
 };
 
-#undef UE_API
diff --git a/Engine/Plugins/Mutable/Source/CustomizableObjectEditor/Private/MuCOE/Widgets/SMutablePinTypeSelector.cpp b/Engine/Plugins/Mutable/Source/CustomizableObjectEditor/Private/MuCOE/Widgets/SMutablePinTypeSelector.cpp
--- a/Engine/Plugins/Mutable/Source/CustomizableObjectEditor/Private/MuCOE/Widgets/SMutablePinTypeSelector.cpp
+++ b/Engine/Plugins/Mutable/Source/CustomizableObjectEditor/Private/MuCOE/Widgets/SMutablePinTypeSelector.cpp
@@ -35,6 +35,7 @@ void SMutablePinTypeSelector::Construct(const FArguments& InArgs)
 			.OptionsSource(&AllowedPinTypes)
 			.InitiallySelectedItem(GetInitiallySelectedPinType())
 			.IsEnabled(this, &SMutablePinTypeSelector::ShouldVariableTypePropertyBeEnabled)
+			.ToolTipText(this, &SMutablePinTypeSelector::GetToolTipText)
 			.OnGenerateWidget(this, &SMutablePinTypeSelector::OnGenerateVariableTypeRow)
 			.OnSelectionChanged(this, &SMutablePinTypeSelector::OnTypeDropdownSelectionChange)
 			.Content()
@@ -80,6 +81,19 @@ bool SMutablePinTypeSelector::ShouldVariableTypePropertyBeEnabled() const
 }
 
 
+FText SMutablePinTypeSelector::GetToolTipText() const
+{
+	if (ShouldVariableTypePropertyBeEnabled())
+	{
+		return LOCTEXT("SMutablePinTypeSelector_SelectorTooltip", "Select the type of pin you want for this pin to be.");
+	}
+	else
+	{
+		return LOCTEXT("SMutablePinTypeSelector_SelectorTooltip_Disabled", "Unable to change the pin type while connected to another pin.");
+	}
+}
+
+
 TSharedRef<SWidget> SMutablePinTypeSelector::OnGenerateVariableTypeRow(FName Type)
 {
 	const FSlateBrush* IconBrush = FAppStyle::GetBrush(TEXT("Kismet.VariableList.TypeIcon"));
diff --git a/Engine/Plugins/Mutable/Source/CustomizableObjectEditor/Private/MuCOE/Widgets/SMutablePinTypeSelector.h b/Engine/Plugins/Mutable/Source/CustomizableObjectEditor/Private/MuCOE/Widgets/SMutablePinTypeSelector.h
--- a/Engine/Plugins/Mutable/Source/CustomizableObjectEditor/Private/MuCOE/Widgets/SMutablePinTypeSelector.h
+++ b/Engine/Plugins/Mutable/Source/CustomizableObjectEditor/Private/MuCOE/Widgets/SMutablePinTypeSelector.h
@@ -29,6 +29,7 @@ private:
 	
 	FName GetInitiallySelectedPinType() const;
 	bool ShouldVariableTypePropertyBeEnabled() const;
+	FText GetToolTipText() const;
 	TSharedRef<SWidget> OnGenerateVariableTypeRow(FName Type);
 	void OnTypeDropdownSelectionChange(FName NewTypeName, ESelectInfo::Type Arg) const;
 	TSharedRef<SWidget> GenerateCurrentSelectedTypeWidget();
