[Mutable] Fix ImagePixelFormat not being used in Bake.

#Synced-CL 52202739

#Synced-CL 52202739

diff --git a/Engine/Plugins/Mutable/Source/CustomizableObject/Private/MuCO/CustomizableObjectSystem.cpp b/Engine/Plugins/Mutable/Source/CustomizableObject/Private/MuCO/CustomizableObjectSystem.cpp
--- a/Engine/Plugins/Mutable/Source/CustomizableObject/Private/MuCO/CustomizableObjectSystem.cpp
+++ b/Engine/Plugins/Mutable/Source/CustomizableObject/Private/MuCO/CustomizableObjectSystem.cpp
@@ -1924,11 +1924,6 @@ namespace impl
 		check(UCustomizableObjectSystem::GetInstance() != nullptr);
 		check(UCustomizableObjectSystem::GetInstance()->GetPrivate() != nullptr);
 		
-		if (OperationData->PixelFormatOverride)
-		{
-			OperationData->MutableSystem->SetImagePixelConversionOverride( OperationData->PixelFormatOverride );
-		}
-
 		check(OperationData->LiveInstance);
 		
 		// Prevent the usage of the already cached elements by fully clearing the cache.
@@ -2330,8 +2325,6 @@ namespace impl
 		const bool bFreeCache = !CVarReuseLiveInstanceCacheBetweenUpdates.GetValueOnAnyThread();
 		MutableSystem->EndUpdate(LiveInstance, bClearRoms, bFreeCache);
 		
-		MutableSystem->SetImagePixelConversionOverride(nullptr);
-
 		UCustomizableObjectSystem::GetInstance()->GetPrivate()->MutableTaskGraph.AllowLaunchingMutableTaskLowPriority(true, true);
 	}
 
@@ -2681,13 +2674,6 @@ namespace impl
 		}
 #endif
 		
-		// Task: Mutable Update and GetMesh
-		//-------------------------------------------------------------
-
-#if WITH_EDITOR
-		Operation->PixelFormatOverride = SystemPrivate->ImageFormatOverrideFunc;
-#endif
-		
 		// CreateMutableInstance
 		{
 			if (FLogBenchmarkUtil::IsBenchmarkingReportingEnabled())
@@ -2730,6 +2716,7 @@ namespace impl
 				ObjectPrivate->SkeletalMeshIdRegistry,
 				Operation->ExternalResourceProvider,
 				Operation->PassthroughObjectLoader,
+				Operation->PixelFormatOverride,
 				UpdateInstanceLogger);
 			
 			Operation->LiveInstance->Parameters = Operation->Parameters;
diff --git a/Engine/Plugins/Mutable/Source/CustomizableObjectEditor/Private/MuCOE/CustomizableObjectInstanceBakingUtils.cpp b/Engine/Plugins/Mutable/Source/CustomizableObjectEditor/Private/MuCOE/CustomizableObjectInstanceBakingUtils.cpp
--- a/Engine/Plugins/Mutable/Source/CustomizableObjectEditor/Private/MuCOE/CustomizableObjectInstanceBakingUtils.cpp
+++ b/Engine/Plugins/Mutable/Source/CustomizableObjectEditor/Private/MuCOE/CustomizableObjectInstanceBakingUtils.cpp
@@ -417,6 +417,8 @@ void ScheduleInstanceUpdateForBaking(UCustomizableObjectInstance& InInstance, FI
 
 	InInstanceUpdateDelegate.AddStatic(&OnInstanceUpdateFinish);
 	
+	PrepareUnrealCompression();
+	
 	// Schedule the update
 	const TSharedRef<FUpdateContextPrivate> Context = MakeShared<FUpdateContextPrivate>(InInstance);
 	Context->bForce = true;
diff --git a/Engine/Plugins/Mutable/Source/MutableRuntime/Internal/MuR/System.h b/Engine/Plugins/Mutable/Source/MutableRuntime/Internal/MuR/System.h
--- a/Engine/Plugins/Mutable/Source/MutableRuntime/Internal/MuR/System.h
+++ b/Engine/Plugins/Mutable/Source/MutableRuntime/Internal/MuR/System.h
@@ -173,6 +173,8 @@ namespace UE::Mutable::Private
 		TSharedPtr<FPassthroughObjectLoader> PassthroughObjectLoader;
 		
 		TSharedPtr<FLiveInstanceLogger> UpdateLogger;
+		
+		FImageOperator::FImagePixelFormatFunc PixelFormatOverride;
 	};
 
 
@@ -264,25 +266,19 @@ namespace UE::Mutable::Private
          * Refer to Settings::SetWorkingMemoryBudget for more information.
 		 */
         UE_API void SetWorkingMemoryBytes( uint64 Bytes );
-		
-		/** Set a function that will be used to convert image pixel formats instead of the internal conversion. 
-		* \warning The provided function can be called from any thread, and also concurrently.
-		* If this function fails (returns false in the first parameter) the internal function is attempted next.
-		* This is useful to provide higher-quality external compressors in editor or when cooking.
-		*/
-		UE_API void SetImagePixelConversionOverride(const FImageOperator::FImagePixelFormatFunc&);
 
         //! Create a new instance from the given model. The instance can then be configured through
         //! calls to BeginUpdate/EndUpdate.
         //! A call to NewInstance must be paired to a call to ReleasesInstance when the instance is
         //! no longer needed.
         UE_API TSharedRef<FLiveInstance> NewInstance(const TSharedPtr<FModel>& Model,
-		const TSharedPtr<UE::Mutable::Private::FMeshIdRegistry>& InMeshIdRegistry,
-		const TSharedPtr<UE::Mutable::Private::FImageIdRegistry>& InImageIdRegistry,
-		const TSharedPtr<UE::Mutable::Private::FMaterialIdRegistry>& InMaterialIdRegistry,
-		const TSharedPtr<UE::Mutable::Private::FSkeletalMeshIdRegistry>& InSkeletalMeshIdRegistry,
+			const TSharedPtr<UE::Mutable::Private::FMeshIdRegistry>& InMeshIdRegistry,
+			const TSharedPtr<UE::Mutable::Private::FImageIdRegistry>& InImageIdRegistry,
+			const TSharedPtr<UE::Mutable::Private::FMaterialIdRegistry>& InMaterialIdRegistry,
+			const TSharedPtr<UE::Mutable::Private::FSkeletalMeshIdRegistry>& InSkeletalMeshIdRegistry,
         	const TSharedPtr<FExternalResourceProvider>& ExternalResourceProvider,
         	const TSharedPtr<FPassthroughObjectLoader>& PassthroughObjectLoader,
+        	const FImageOperator::FImagePixelFormatFunc& PixelFormatOverride,
 			const TSharedPtr<FLiveInstanceLogger>& LiveInstanceLogger);
 
 		/** 
@@ -391,10 +387,7 @@ namespace UE::Mutable::Private
 		* This should only happen in-editor.
 		*/
 		bool bUnrecoverableError = false;
-
-		/** If this is set, it will be tried first instead of the internal formatting function. */
-		FImageOperator::FImagePixelFormatFunc ImagePixelFormatOverride;
-
+		
 		FWorkingMemoryManager WorkingMemoryManager;
 
 		/** Increased each time a rom is loaded. Used to keep track of which roms to evict. */
diff --git a/Engine/Plugins/Mutable/Source/MutableRuntime/Private/MuR/CodeRunner.cpp b/Engine/Plugins/Mutable/Source/MutableRuntime/Private/MuR/CodeRunner.cpp
--- a/Engine/Plugins/Mutable/Source/MutableRuntime/Private/MuR/CodeRunner.cpp
+++ b/Engine/Plugins/Mutable/Source/MutableRuntime/Private/MuR/CodeRunner.cpp
@@ -8663,7 +8663,7 @@ namespace UE::Mutable::Private
 				return New;
 			},
 
-			Runner->System.ImagePixelFormatOverride
+			Runner->LiveInstance->PixelFormatOverride
 		);
 	}
 }
diff --git a/Engine/Plugins/Mutable/Source/MutableRuntime/Private/MuR/CodeRunnerTasked.cpp b/Engine/Plugins/Mutable/Source/MutableRuntime/Private/MuR/CodeRunnerTasked.cpp
--- a/Engine/Plugins/Mutable/Source/MutableRuntime/Private/MuR/CodeRunnerTasked.cpp
+++ b/Engine/Plugins/Mutable/Source/MutableRuntime/Private/MuR/CodeRunnerTasked.cpp
@@ -961,7 +961,7 @@ namespace UE::Mutable::Private
 			return false;
 		}
 
-		ImagePixelFormatFunc = Runner->System.ImagePixelFormatOverride;
+		ImagePixelFormatFunc = Runner->LiveInstance->PixelFormatOverride;
 
 		// Create destination data
 		Result = Runner->CreateImage(Base->GetSizeX(), Base->GetSizeY(), Base->GetLODCount(), TargetFormat, EInitializationType::NotInitialized);
@@ -1078,7 +1078,7 @@ namespace UE::Mutable::Private
 		FImageOperator ImOp = MakeImageOperator(Runner);
 		ImOp.ImageMipmap_PrepareScratch(Result.Get(), StartLevel, LevelCount, Scratch);
 
-		ImagePixelFormatFunc = Runner->System.ImagePixelFormatOverride;
+		ImagePixelFormatFunc = Runner->LiveInstance->PixelFormatOverride;
 
 		return true;
 	}
@@ -1405,7 +1405,7 @@ namespace UE::Mutable::Private
 
 		Result = Runner->CreateImage( destSize[0], destSize[1], Lods, Base->GetFormat(), EInitializationType::NotInitialized );
 
-		ImagePixelFormatFunc = Runner->System.ImagePixelFormatOverride;
+		ImagePixelFormatFunc = Runner->LiveInstance->PixelFormatOverride;
 
 		return true;
 	}
@@ -1526,7 +1526,7 @@ namespace UE::Mutable::Private
 
 		Result = Runner->CreateImage(DestSize[0], DestSize[1], Lods, Base->GetFormat(), EInitializationType::NotInitialized);
 
-		ImagePixelFormatFunc = Runner->System.ImagePixelFormatOverride;
+		ImagePixelFormatFunc = Runner->LiveInstance->PixelFormatOverride;
 
 		return true;
 	}
@@ -1752,7 +1752,7 @@ namespace UE::Mutable::Private
 			}
 		}
 
-		ImagePixelFormatFunc = Runner->System.ImagePixelFormatOverride;
+		ImagePixelFormatFunc = Runner->LiveInstance->PixelFormatOverride;
 
 		return true;
 	}
diff --git a/Engine/Plugins/Mutable/Source/MutableRuntime/Private/MuR/System.cpp b/Engine/Plugins/Mutable/Source/MutableRuntime/Private/MuR/System.cpp
--- a/Engine/Plugins/Mutable/Source/MutableRuntime/Private/MuR/System.cpp
+++ b/Engine/Plugins/Mutable/Source/MutableRuntime/Private/MuR/System.cpp
@@ -112,15 +112,7 @@ namespace UE::Mutable::Private
 
 		UpdateStats();
 	}
-
 	
-	void FSystem::SetImagePixelConversionOverride(const FImageOperator::FImagePixelFormatFunc& FormatFunc)
-	{
-		LLM_SCOPE_BYNAME(TEXT("MutableRuntime"));
-
-		ImagePixelFormatOverride = FormatFunc;
-	}
-
    
     TSharedRef<FLiveInstance> FSystem::NewInstance( const TSharedPtr<FModel>& InModel,
 		const TSharedPtr<UE::Mutable::Private::FMeshIdRegistry>& InMeshIdRegistry,
@@ -129,6 +121,7 @@ namespace UE::Mutable::Private
 		const TSharedPtr<UE::Mutable::Private::FSkeletalMeshIdRegistry>& InSkeletalMeshIdRegistry,
     	const TSharedPtr<FExternalResourceProvider>& ExternalResourceProvider,
     	const TSharedPtr<FPassthroughObjectLoader>& PassthroughObjectLoader,
+    	const FImageOperator::FImagePixelFormatFunc& PixelFormatOverride,
 		const TSharedPtr<FLiveInstanceLogger>& LiveInstanceLogger)
     {
 		LLM_SCOPE_BYNAME(TEXT("MutableRuntime"));
@@ -145,6 +138,7 @@ namespace UE::Mutable::Private
     	InstanceData->ImageIdRegistry = InImageIdRegistry;
     	InstanceData->MaterialIdRegistry = InMaterialIdRegistry;
     	InstanceData->SkeletalMeshIdRegistry = InSkeletalMeshIdRegistry;
+    	InstanceData->PixelFormatOverride = PixelFormatOverride;
     	InstanceData->UpdateLogger = LiveInstanceLogger;
     	
     	check(InstanceData->MeshIdRegistry.IsValid())
diff --git a/Engine/Plugins/Mutable/Source/MutableTools/Private/MuT/CodeOptimiser.cpp b/Engine/Plugins/Mutable/Source/MutableTools/Private/MuT/CodeOptimiser.cpp
--- a/Engine/Plugins/Mutable/Source/MutableTools/Private/MuT/CodeOptimiser.cpp
+++ b/Engine/Plugins/Mutable/Source/MutableTools/Private/MuT/CodeOptimiser.cpp
@@ -171,9 +171,7 @@ namespace UE::Mutable::Private
 			Settings.SetProfile( false );
 			Settings.SetImageCompressionQuality( ImageCompressionQuality );
 			TSharedPtr<FSystem> System = MakeShared<FSystem>(Settings);
-
-			System->ImagePixelFormatOverride = ImOp.FormatImageOverride;
-
+			
 			FSourceDataDescriptor SourceDataDescriptor;
 			if (DataType == EDataType::Image || DataType == EDataType::Mesh)
 			{
@@ -196,7 +194,8 @@ namespace UE::Mutable::Private
 			TSharedPtr<FParameters> LocalParams = FModel::NewParameters(Model);
 
 			TSharedRef<FLiveInstance> TempLiveInstance = System->NewBuildInstance(Model, LocalParams,  nullptr);
-			
+			TempLiveInstance->PixelFormatOverride = ImOp.FormatImageOverride;
+
 			// Calculate the value and replace this op by a constant
 			switch( DataType )
 			{
