[Mutable] Fix GeneratedMaterials cache bug.

- opertaor== was not defined by the GeneratedMaterialsCache key but the parent class oprerator== was used.
- Passed options now don't use inheritance to preventing similar issues from happening in the future.

#rnx
#Synced-CL 52208093

diff --git a/Engine/Plugins/Mutable/Source/MutableTools/Private/MuT/CodeGenerator.cpp b/Engine/Plugins/Mutable/Source/MutableTools/Private/MuT/CodeGenerator.cpp
--- a/Engine/Plugins/Mutable/Source/MutableTools/Private/MuT/CodeGenerator.cpp
+++ b/Engine/Plugins/Mutable/Source/MutableTools/Private/MuT/CodeGenerator.cpp
@@ -1179,7 +1179,7 @@ namespace UE::Mutable::Private
 					check(SurfaceNode->Material)
 
 					FMaterialGenerationOptions MaterialOptions;
-					MaterialOptions.State = Options.State;
+					MaterialOptions.GenericOptions.State = Options.State;
 					MaterialOptions.ComponentId = ComponentId;
 					MaterialOptions.LODIndex = Options.LODIndex;
 					MaterialOptions.Tags = SurfaceNode->Tags;
@@ -2167,7 +2167,7 @@ namespace UE::Mutable::Private
 					FImageGenerationOptions ClipOptions(StaticOptions.ComponentId, StaticOptions.LODIndex);
 					ClipOptions.ImageLayoutStrategy = CompilerOptions::TextureLayoutStrategy::None;
 					ClipOptions.LayoutBlockId = FLayoutBlock::InvalidBlockId;
-					ClipOptions.State = StaticOptions.State;
+					ClipOptions.GenericOptions.State = StaticOptions.State;
 
 					FImageGenerationResult ClipMaskResult;
 					GenerateImage(ClipOptions, ClipMaskResult, TypedClipNode->ClipMask);
@@ -2682,9 +2682,9 @@ namespace UE::Mutable::Private
 				{
 					// Generate the image block
 					FImageGenerationOptions ImageOptions(Options.ComponentId,Options.LODIndex);
-					ImageOptions.State = Options.State;
+					ImageOptions.GenericOptions.State = Options.State;
 					ImageOptions.ImageLayoutStrategy = ImageLayoutStrategy;
-					ImageOptions.ActiveTags = Edit->EnableTags; // TODO: Merge with current tags?
+					ImageOptions.GenericOptions.ActiveTags = Edit->EnableTags; // TODO: Merge with current tags?
 					ImageOptions.RectSize = { 0,0 };
 					ImageOptions.LayoutToApply = ExtendLayout;
 					ImageOptions.LayoutBlockId = ExtendLayout->Blocks[BlockIndex].Id;
diff --git a/Engine/Plugins/Mutable/Source/MutableTools/Private/MuT/CodeGenerator.h b/Engine/Plugins/Mutable/Source/MutableTools/Private/MuT/CodeGenerator.h
--- a/Engine/Plugins/Mutable/Source/MutableTools/Private/MuT/CodeGenerator.h
+++ b/Engine/Plugins/Mutable/Source/MutableTools/Private/MuT/CodeGenerator.h
@@ -171,13 +171,14 @@ namespace UE::Mutable::Private
 				return KeyHash;
 			}
 
-			bool operator==(const FGenericGenerationOptions& InKey) const = default;
+			friend bool operator==(const FGenericGenerationOptions&, const FGenericGenerationOptions&) = default;
 
 			int32 State = -1;
 			TArray<FString> ActiveTags;
 
 			// Variable to identify if an operation will end up generating an image.
 			bool bIsImage = false;
+			
 		};
 
 		struct FObjectGenerationOptions : public FGenericGenerationOptions
@@ -468,15 +469,17 @@ namespace UE::Mutable::Private
 		// Images
 		
 		/** Options that affect the generation of images. It is like list of what required data we want while parsing down the image node graph. */
-		struct FImageGenerationOptions : public FGenericGenerationOptions
+		struct FImageGenerationOptions
 		{
 			FImageGenerationOptions(int32 InComponentId, int32 InLODIndex)
 				: ComponentId(InComponentId)
 				, LODIndex(InLODIndex)
 			{
-				bIsImage = true;
+				GenericOptions.bIsImage = true;
 			}
 
+			FGenericGenerationOptions GenericOptions;
+
 			/** The id of the component that we are currently generating. */
 			int32 ComponentId = -1;
 
@@ -498,7 +501,7 @@ namespace UE::Mutable::Private
 
 			friend FORCEINLINE uint32 GetTypeHash(const FImageGenerationOptions& InKey)
 			{
-				uint32 KeyHash = GetTypeHash((FGenericGenerationOptions&)InKey);
+				uint32 KeyHash = GetTypeHash((FGenericGenerationOptions&)InKey.GenericOptions);
 				KeyHash = HashCombineFast(KeyHash, ::GetTypeHash(InKey.ComponentId));
 				KeyHash = HashCombineFast(KeyHash, ::GetTypeHash(InKey.LODIndex));
 				KeyHash = HashCombineFast(KeyHash, ::GetTypeHash(InKey.ImageLayoutStrategy));
@@ -609,8 +612,10 @@ namespace UE::Mutable::Private
 		// Materials	
 
 		/** Options that affect the generation of Materials. It is like list of what required data we want while parsing down the material node graph. */
-		struct FMaterialGenerationOptions : public FGenericGenerationOptions
+		struct FMaterialGenerationOptions
 		{
+			FGenericGenerationOptions GenericOptions;
+
 			/** The id of the component that we are currently generating. */
 			int32 ComponentId = -1;
 
@@ -636,6 +641,18 @@ namespace UE::Mutable::Private
 			/** Layout block that we are trying to generate if any. */
 			uint64 LayoutBlockId = FLayoutBlock::InvalidBlockId;
 			TManagedPtr<const FLayout> LayoutToApply;
+
+			friend bool operator==(const FMaterialGenerationOptions&, const FMaterialGenerationOptions&) = default;
+
+			friend FORCEINLINE uint32 GetTypeHash(const FMaterialGenerationOptions& Options)
+			{
+				uint32 KeyHash = 0;
+
+				// TODO: Use a more representative Hash.
+				KeyHash = HashCombineFast(KeyHash, GetTypeHash(Options.GenericOptions));
+
+				return KeyHash;
+			}
 		};
 
 		/** */
@@ -659,7 +676,7 @@ namespace UE::Mutable::Private
 				return KeyHash;
 			}
 
-			FORCEINLINE bool operator==(const FGeneratedMaterialCacheKey& Other) const = default;
+			friend bool operator==(const FGeneratedMaterialCacheKey&, const FGeneratedMaterialCacheKey&) = default;
 		};
     	
 		typedef TMap<FGeneratedMaterialCacheKey, FMaterialGenerationResult> FGeneratedMaterialsMap;
@@ -953,7 +970,7 @@ namespace UE::Mutable::Private
     			
     			if (ImageOptions.IsSet())
     			{
-    				return static_cast<FGenericGenerationOptions>(*ImageOptions);
+    				return ImageOptions->GenericOptions;
     			}
     			
     			if (ScalarOptions.IsSet())
@@ -968,7 +985,7 @@ namespace UE::Mutable::Private
     			
     			if (MaterialOptions.IsSet())
     			{
-    				return static_cast<FGenericGenerationOptions>(*MaterialOptions);
+    				return MaterialOptions->GenericOptions;
     			}
     			
     			return {};
diff --git a/Engine/Plugins/Mutable/Source/MutableTools/Private/MuT/CodeGenerator_Color.cpp b/Engine/Plugins/Mutable/Source/MutableTools/Private/MuT/CodeGenerator_Color.cpp
--- a/Engine/Plugins/Mutable/Source/MutableTools/Private/MuT/CodeGenerator_Color.cpp
+++ b/Engine/Plugins/Mutable/Source/MutableTools/Private/MuT/CodeGenerator_Color.cpp
@@ -338,8 +338,8 @@ namespace UE::Mutable::Private
 		int32 ComponentId = -1; // TODO.
 		int32 LODIndex = 0; // TODO.
 		FImageGenerationOptions ImageOptions(ComponentId, LODIndex);
-		ImageOptions.State = Options.State;
-		ImageOptions.ActiveTags = Options.ActiveTags;
+		ImageOptions.GenericOptions.State = Options.State;
+		ImageOptions.GenericOptions.ActiveTags = Options.ActiveTags;
 
 		Ptr<ASTOp> base;
 		if (node.Image)
diff --git a/Engine/Plugins/Mutable/Source/MutableTools/Private/MuT/CodeGenerator_External.cpp b/Engine/Plugins/Mutable/Source/MutableTools/Private/MuT/CodeGenerator_External.cpp
--- a/Engine/Plugins/Mutable/Source/MutableTools/Private/MuT/CodeGenerator_External.cpp
+++ b/Engine/Plugins/Mutable/Source/MutableTools/Private/MuT/CodeGenerator_External.cpp
@@ -263,23 +263,23 @@ namespace UE::Mutable::Private
 
 				if (Options.MeshStaticOptions.IsSet())
 				{
-					GenericOptions = static_cast<FGenericGenerationOptions>(Options.MeshStaticOptions.GetValue());
+					GenericOptions = Options.MeshStaticOptions.GetValue();
 				}
 				else if (Options.ColourOptions.IsSet())
 				{
-					GenericOptions = static_cast<FGenericGenerationOptions>(Options.ColourOptions.GetValue());
+					GenericOptions = Options.ColourOptions.GetValue();
 				}
 				else if (Options.ImageOptions.IsSet())
 				{
-					GenericOptions = static_cast<FGenericGenerationOptions>(Options.ImageOptions.GetValue());
+					GenericOptions = Options.ImageOptions.GetValue().GenericOptions;
 				}
 				else if (Options.MaterialOptions.IsSet())
 				{
-					GenericOptions = static_cast<FGenericGenerationOptions>(Options.MaterialOptions.GetValue());
+					GenericOptions = Options.MaterialOptions.GetValue().GenericOptions;
 				}
 				else if (Options.ScalarOptions.IsSet())
 				{
-					GenericOptions = static_cast<FGenericGenerationOptions>(Options.ScalarOptions.GetValue());
+					GenericOptions = Options.ScalarOptions.GetValue();
 				}
 				
 				FRangeGenerationResult rangeResult;
diff --git a/Engine/Plugins/Mutable/Source/MutableTools/Private/MuT/CodeGenerator_FirstPass.h b/Engine/Plugins/Mutable/Source/MutableTools/Private/MuT/CodeGenerator_FirstPass.h
--- a/Engine/Plugins/Mutable/Source/MutableTools/Private/MuT/CodeGenerator_FirstPass.h
+++ b/Engine/Plugins/Mutable/Source/MutableTools/Private/MuT/CodeGenerator_FirstPass.h
@@ -42,11 +42,7 @@ namespace UE::Mutable::Private
 		TManagedPtr<const FLayout> Layout;
 		Ptr<const NodeLayout> Source;
 
-		FORCEINLINE bool operator==(const FGeneratedLayout& Other) const
-		{
-			return Layout == Other.Layout
-				&& Source == Other.Source;
-		}
+		friend bool operator==(const FGeneratedLayout&, const FGeneratedLayout&) = default;
 	};
 
 	/** Store the results of the code generation of a mesh. */
@@ -72,10 +68,17 @@ namespace UE::Mutable::Private
 			TArray<FGeneratedLayout> GeneratedLayouts;
 			Ptr<ASTOp> Condition;
 			Ptr<ASTOp> MeshFragment;
+
+			friend bool operator==(const FExtraLayouts&, const FExtraLayouts&) = default;
 		};
 		TArray<FExtraLayouts> ExtraMeshLayouts;
+
+		friend bool operator==(const FMeshGenerationResult&, const FMeshGenerationResult&) = default;
 	};
 
+	
+
+
 	using FMeshTask = UE::Tasks::TTask<FMeshGenerationResult>;
 
 	struct FSurfaceGenerationResult
diff --git a/Engine/Plugins/Mutable/Source/MutableTools/Private/MuT/CodeGenerator_Image.cpp b/Engine/Plugins/Mutable/Source/MutableTools/Private/MuT/CodeGenerator_Image.cpp
--- a/Engine/Plugins/Mutable/Source/MutableTools/Private/MuT/CodeGenerator_Image.cpp
+++ b/Engine/Plugins/Mutable/Source/MutableTools/Private/MuT/CodeGenerator_Image.cpp
@@ -302,7 +302,7 @@ namespace UE::Mutable::Private
 			for (int32 a = 0; a < Node->Ranges.Num(); ++a)
 			{
 				FRangeGenerationResult rangeResult;
-				GenerateRange(rangeResult, Options, Node->Ranges[a]);
+				GenerateRange(rangeResult, Options.GenericOptions, Node->Ranges[a]);
 				op->Ranges.Emplace(op.get(), rangeResult.sizeOp, rangeResult.rangeName, rangeResult.rangeUID);
 			}
 		}
@@ -445,7 +445,7 @@ namespace UE::Mutable::Private
         if (Node->Colour )
         {
 			FColorGenerationResult ColorResult;
-            GenerateColor(ColorResult, Options, Node->Colour);
+            GenerateColor(ColorResult, Options.GenericOptions, Node->Colour);
 			colour = ColorResult.op;
         }
         else
@@ -533,7 +533,7 @@ namespace UE::Mutable::Private
         if (Node->Range )
         {
             FRangeGenerationResult rangeResult;
-            GenerateRange( rangeResult, Options, Node->Range );
+            GenerateRange( rangeResult, Options.GenericOptions, Node->Range );
 
             op->range.rangeSize = rangeResult.sizeOp;
             op->range.rangeName = rangeResult.rangeName;
@@ -621,7 +621,7 @@ namespace UE::Mutable::Private
 		if (Node->OffsetX)
 		{
 			FScalarGenerationResult ScalarResult;
-			GenerateScalar(ScalarResult, Options, Node->OffsetX);
+			GenerateScalar(ScalarResult, Options.GenericOptions, Node->OffsetX);
 			OffsetX = ScalarResult.op;
 		}
 
@@ -629,7 +629,7 @@ namespace UE::Mutable::Private
 		if (Node->OffsetY)
 		{
 			FScalarGenerationResult ScalarResult;
-			GenerateScalar(ScalarResult, Options, Node->OffsetY);
+			GenerateScalar(ScalarResult, Options.GenericOptions, Node->OffsetY);
 			OffsetY = ScalarResult.op;
 		}
 	
@@ -637,7 +637,7 @@ namespace UE::Mutable::Private
 		if (Node->ScaleX)
 		{
 			FScalarGenerationResult ScalarResult;
-			GenerateScalar(ScalarResult, Options, Node->ScaleX);
+			GenerateScalar(ScalarResult, Options.GenericOptions, Node->ScaleX);
 			ScaleX = ScalarResult.op;
 		}
 	
@@ -645,7 +645,7 @@ namespace UE::Mutable::Private
 		if (Node->ScaleY)
 		{
 			FScalarGenerationResult ScalarResult;
-			GenerateScalar(ScalarResult, Options, Node->ScaleY);
+			GenerateScalar(ScalarResult, Options.GenericOptions, Node->ScaleY);
 			ScaleY = ScalarResult.op;
 		}
 
@@ -653,7 +653,7 @@ namespace UE::Mutable::Private
 		if (Node->Rotation)
 		{
 			FScalarGenerationResult ScalarResult;
-			GenerateScalar(ScalarResult, Options, Node->Rotation);
+			GenerateScalar(ScalarResult, Options.GenericOptions, Node->Rotation);
 			Rotation = ScalarResult.op;
 		}
 
@@ -753,7 +753,7 @@ namespace UE::Mutable::Private
         if ( NodeScalar* Factor = Node->Factor.get() )
         {
 			FScalarGenerationResult ParamResult;
-			GenerateScalar(ParamResult, Options, Factor);
+			GenerateScalar(ParamResult, Options.GenericOptions, Factor);
 			op->Factor = ParamResult.op;
         }
         else
@@ -1000,7 +1000,7 @@ namespace UE::Mutable::Private
         if ( NodeScalar* Factor = Node->Factor.get() )
         {
 			FScalarGenerationResult ParamResult;
-			GenerateScalar(ParamResult, Options, Factor);
+			GenerateScalar(ParamResult, Options.GenericOptions, Factor);
 			op->Factor = ParamResult.op;
 		}
         else
@@ -1138,7 +1138,7 @@ namespace UE::Mutable::Private
         if ( NodeScalar* ThresholdNode = Node->Threshold.get() )
         {
 			FScalarGenerationResult ParamResult;
-			GenerateScalar(ParamResult, Options, ThresholdNode);
+			GenerateScalar(ParamResult, Options.GenericOptions, ThresholdNode);
 			ThresholdOp = ParamResult.op;
         }
         else
@@ -1238,7 +1238,7 @@ namespace UE::Mutable::Private
         if (Node->Colour )
         {
 			FColorGenerationResult ColorResult;
-            GenerateColor( ColorResult, Options, Node->Colour.get() );
+            GenerateColor( ColorResult, Options.GenericOptions, Node->Colour.get() );
 			base = ColorResult.op;
         }
         else
@@ -1311,7 +1311,7 @@ namespace UE::Mutable::Private
 		if (Node->Parameter )
 		{
 			FScalarGenerationResult ParamResult;
-			GenerateScalar(ParamResult, Options, Node->Parameter.get());
+			GenerateScalar(ParamResult, Options.GenericOptions, Node->Parameter.get());
 			op->Variable = ParamResult.op;
 		}
 		else
@@ -1357,7 +1357,7 @@ namespace UE::Mutable::Private
         if (Node->Parameter)
         {
 			FBoolGenerationResult ParamResult;
-			GenerateBool(ParamResult, Options, Node->Parameter.get());
+			GenerateBool(ParamResult, Options.GenericOptions, Node->Parameter.get());
 			op->condition = ParamResult.op;
 		}
         else
@@ -1397,7 +1397,7 @@ namespace UE::Mutable::Private
         FProjectorGenerationResult projectorResult;
         if (Node->Projector )
         {
-            GenerateProjector( projectorResult, Options, Node->Projector );
+            GenerateProjector( projectorResult, Options.GenericOptions, Node->Projector );
             //projectorAt = Generate( Node->Projector.get() );
         }
         else
@@ -1421,8 +1421,8 @@ namespace UE::Mutable::Private
 			// TODO: This will probably Result in a duplicated mesh subgraph, with the original mesh but new layout block ids.
 			// See if it can be optimized and try to reuse the existing layout block ids instead of generating new ones.
 			FMeshGenerationStaticOptions MeshStaticOptions(Options.ComponentId, Options.LODIndex);
-			MeshStaticOptions.State = Options.State;
-			MeshStaticOptions.ActiveTags = Options.ActiveTags;
+			MeshStaticOptions.State = Options.GenericOptions.State;
+			MeshStaticOptions.ActiveTags = Options.GenericOptions.ActiveTags;
 			FMeshGenerationDynamicOptions MeshOptions;
 			MeshOptions.bLayouts = true;			// We need the layout that we will use to render
 			MeshOptions.bNormalizeUVs = true;		// We need normalized UVs for the projection
@@ -1667,7 +1667,7 @@ namespace UE::Mutable::Private
 			}
 
 			FColorGenerationResult ParamResult;
-			GenerateColor(ParamResult, Options, pPropsNode);
+			GenerateColor(ParamResult, Options.GenericOptions, pPropsNode);
 			ImageRasterOp->angleFadeProperties = ParamResult.op;
 		}
 
@@ -1972,7 +1972,7 @@ namespace UE::Mutable::Private
 		FMaterialGenerationOptions MaterialSourceOptions;
 		MaterialSourceOptions.ComponentId = InOptions.ComponentId;
 		MaterialSourceOptions.LODIndex = InOptions.LODIndex;
-		MaterialSourceOptions.State = InOptions.State;
+		MaterialSourceOptions.GenericOptions.State = InOptions.GenericOptions.State;
 		MaterialSourceOptions.ImageLayoutStrategy = InOptions.ImageLayoutStrategy;
 		MaterialSourceOptions.RectSize = InOptions.RectSize;
 		MaterialSourceOptions.LayoutBlockId = InOptions.LayoutBlockId;
diff --git a/Engine/Plugins/Mutable/Source/MutableTools/Private/MuT/CodeGenerator_Material.cpp b/Engine/Plugins/Mutable/Source/MutableTools/Private/MuT/CodeGenerator_Material.cpp
--- a/Engine/Plugins/Mutable/Source/MutableTools/Private/MuT/CodeGenerator_Material.cpp
+++ b/Engine/Plugins/Mutable/Source/MutableTools/Private/MuT/CodeGenerator_Material.cpp
@@ -165,7 +165,7 @@ namespace UE::Mutable::Private
 		if (Node->Parameter)
 		{
 			FScalarGenerationResult ParamResult;
-			GenerateScalar(ParamResult, Options, Node->Parameter.get());
+			GenerateScalar(ParamResult, Options.GenericOptions, Node->Parameter.get());
 			Variable = ParamResult.op;
 		}
 		else
@@ -315,7 +315,7 @@ namespace UE::Mutable::Private
 		for (const Ptr<NodeRange>& RangeNode : InNode->Ranges)
 		{
 			FRangeGenerationResult RangeResult;
-			GenerateRange(RangeResult, InOptions, RangeNode);
+			GenerateRange(RangeResult, InOptions.GenericOptions, RangeNode);
 			MaterialOp->Ranges.Emplace(MaterialOp.get(), RangeResult.sizeOp, RangeResult.rangeName, RangeResult.rangeUID);
 		}
 
@@ -358,7 +358,7 @@ namespace UE::Mutable::Private
 			for (const TPair<FParameterKey, FImageParameterData>& ImageParameterData : Node->ImageParameters)
 			{
 				FImageGenerationResult ImageResult;
-				GenerateMaterialImage(ImageParameterData, ImageResult, Options.Modifiers, Options.State, Options.ComponentId, Options.LODIndex, Options.Tags, Options.MeshResults, Options.LayoutFromExtension, Node->GetMessageContext());
+				GenerateMaterialImage(ImageParameterData, ImageResult, Options.Modifiers, Options.GenericOptions.State, Options.ComponentId, Options.LODIndex, Options.Tags, Options.MeshResults, Options.LayoutFromExtension, Node->GetMessageContext());
 				MaterialModify->ParametersToModify.Add(ImageParameterData.Key, ASTChild(MaterialModify, ImageResult.op));
 				MaterialModify->ImagePropertyIndexMap.Add(ImageParameterData.Key, ImageParameterData.Value.ImagePropertyIndex);
 			}
@@ -417,7 +417,7 @@ namespace UE::Mutable::Private
 		Ptr<ASTOpMaterialSkeletalMeshBreak> MaterialFromMeshSlotOp = new ASTOpMaterialSkeletalMeshBreak();
 		MaterialFromMeshSlotOp->SlotName = InNode->SlotName;
 		
-		FSkeletalMeshGenerationOptions SkeletalMeshOptions{InOptions, nullptr};
+		FSkeletalMeshGenerationOptions SkeletalMeshOptions{InOptions.GenericOptions, nullptr};
 		FSkeletalMeshGenerationResult SkeletalMeshGenerationResult;
 		GenerateSkeletalMesh(SkeletalMeshGenerationResult, SkeletalMeshOptions, InNode->Mesh);
 		MaterialFromMeshSlotOp->SkeletalMesh = SkeletalMeshGenerationResult.SkeletalMeshOp;
@@ -532,9 +532,9 @@ namespace UE::Mutable::Private
 			{
 				// Generate the image
 				FImageGenerationOptions ImageOptions(ComponentId, LODIndex);
-				ImageOptions.State = State;
+				ImageOptions.GenericOptions.State = State;
 				ImageOptions.ImageLayoutStrategy = ImageLayoutStrategy;
-				ImageOptions.ActiveTags = Tags;
+				ImageOptions.GenericOptions.ActiveTags = Tags;
 				ImageOptions.RectSize = { 0, 0 };
 
 				// TODO: To tasks
@@ -647,10 +647,10 @@ namespace UE::Mutable::Private
 						{
 							// Generate the image
 							FImageGenerationOptions ImageOptions(ComponentId, LODIndex);
-							ImageOptions.State = State;
+							ImageOptions.GenericOptions.State = State;
 							ImageOptions.ImageLayoutStrategy = ImageLayoutStrategy;
 							ImageOptions.RectSize = { 0,0 };
-							ImageOptions.ActiveTags = Tags;
+							ImageOptions.GenericOptions.ActiveTags = Tags;
 							ImageOptions.LayoutToApply = pLayout;
 							ImageOptions.LayoutBlockId = pLayout->Blocks[BlockIndex].Id;
 							FImageGenerationResult ImageResult;
