[Mutable] Improve OpImageMipmap by rounding the results.

#rnx
#Synced-CL 52250858

diff --git a/Engine/Plugins/Mutable/Source/MutableRuntime/Private/MuR/OpImageMipmap.cpp b/Engine/Plugins/Mutable/Source/MutableRuntime/Private/MuR/OpImageMipmap.cpp
--- a/Engine/Plugins/Mutable/Source/MutableRuntime/Private/MuR/OpImageMipmap.cpp
+++ b/Engine/Plugins/Mutable/Source/MutableRuntime/Private/MuR/OpImageMipmap.cpp
@@ -2,7 +2,6 @@
 
 #include "MuR/ImagePrivate.h"
 #include "MuR/ImageFormatUtils.h"
-#include "MuR/System.h"
 #include "MuR/ParallelExecutionUtils.h"
 #include "HAL/UnrealMemory.h"
 #include "Async/ParallelFor.h"
@@ -81,22 +80,20 @@ namespace OpImageMipmapInternal
 					FMemory::Memcpy(&Row0Bits, SourceRow0, sizeof(uint64));
 					FMemory::Memcpy(&Row1Bits, SourceRow1, sizeof(uint64));
 					
-					const uint64 XorRow0Row1Bits = Row0Bits ^ Row1Bits;
-
-					// Average of 2 unsigned integers without overflow extended to work on multiple bytes.
-					constexpr uint64 ShiftMask = 0xFEFEFEFEFEFEFEFE;
-					const uint64 ErrorCorrection = XorRow0Row1Bits & 0x0101010101010101;
-					const uint64 AvgLowBits = (Row0Bits & Row1Bits) + ((XorRow0Row1Bits & ShiftMask) >> 1) + ErrorCorrection;
+					// Average of 2 unsigned integers with rounding without overflow extended to work on multiple bytes.
+					constexpr uint64 ShiftMask = 0x7F7F7F7F7F7F7F7F;
+			
+					const uint64 AvgLowBits = (Row0Bits | Row1Bits) - (((Row0Bits ^ Row1Bits) >> 1) & ShiftMask);
 					const uint64 AvgHighBits = AvgLowBits >> 32;
-					const uint32 Result = (AvgLowBits & AvgHighBits) + (((AvgLowBits ^ AvgHighBits) & ShiftMask) >> 1);
-				
+					const uint32 Result = (AvgLowBits | AvgHighBits) - (((AvgLowBits ^ AvgHighBits) >> 1) & ShiftMask);
+
 					FMemory::Memcpy(DestRow, &Result, sizeof(uint32));
 				}
 				else
 				{
 					for (int32 C = 0; C < PIXEL_SIZE; ++C)
 					{
-						int32 PixelSum = SourceRow0[C] + SourceRow0[PIXEL_SIZE + C] + SourceRow1[C] + SourceRow1[PIXEL_SIZE + C];
+						int32 PixelSum = SourceRow0[C] + SourceRow0[PIXEL_SIZE + C] + SourceRow1[C] + SourceRow1[PIXEL_SIZE + C] + 2;
 						DestRow[C] = (uint8)(PixelSum >> 2);
 					}
 				}
@@ -116,9 +113,9 @@ namespace OpImageMipmapInternal
 					FMemory::Memcpy(&Row0Bits, SourceRow0, sizeof(uint32));
 					FMemory::Memcpy(&Row1Bits, SourceRow1, sizeof(uint32));
 
-					// Average of 2 unsigned integers without overflow extended to work on multiple bytes.
-					constexpr uint32 ShiftMask = 0xFEFEFEFE;
-					const uint32 Result = (Row0Bits & Row1Bits) + (((Row0Bits ^ Row1Bits) & ShiftMask) >> 1);
+					// Average of 2 unsigned integers with rounding without overflow extended to work on multiple bytes.
+					constexpr uint32 ShiftMask = 0x7F7F7F7F;
+					const uint32 Result = (Row0Bits | Row1Bits) - (((Row0Bits ^ Row1Bits) >> 1) & ShiftMask);
 
 					FMemory::Memcpy(DestRow, &Result, sizeof(uint32));
 				}
@@ -126,7 +123,7 @@ namespace OpImageMipmapInternal
 				{
 					for (int32 C = 0; C < PIXEL_SIZE; ++C)
 					{
-						int32 PixelSum = SourceRow0[C] + SourceRow1[C];
+						int32 PixelSum = SourceRow0[C] + SourceRow1[C] + 1;
 						DestRow[C] = (uint8)(PixelSum >> 1);
 					}
 				}
@@ -162,10 +159,9 @@ namespace OpImageMipmapInternal
 					FMemory::Memcpy(&Col0Bits, SourceRow0, sizeof(uint32));
 					FMemory::Memcpy(&Col1Bits, SourceRow0 + 4, sizeof(uint32));
 
-					// Average of 2 unsigned integers without overflow extended to work on multiple bytes. 
-					// In this case we use the ceil variant to be consistent with the method used for 4 pixel average. 
-					constexpr uint32 ShiftMask = 0xFEFEFEFE;
-					const uint32 Result = (Col0Bits & Col1Bits) + (((Col0Bits ^ Col1Bits) & ShiftMask) >> 1);
+					// Average of 2 unsigned integers with rounding without overflow extended to work on multiple bytes.
+					constexpr uint32 ShiftMask = 0x7F7F7F7F;
+					const uint32 Result = (Col0Bits | Col1Bits) - (((Col0Bits ^ Col1Bits) >> 1) & ShiftMask);
 
 					FMemory::Memcpy(DestRow, &Result, sizeof(uint32));
 				}
@@ -173,7 +169,7 @@ namespace OpImageMipmapInternal
 				{
 					for (int32 C = 0; C < PIXEL_SIZE; ++C)
 					{
-						int32 P = SourceRow0[C] + SourceRow0[PIXEL_SIZE + C];
+						int32 P = SourceRow0[C] + SourceRow0[PIXEL_SIZE + C] + 1;
 						DestRow[C] = (uint8)(P >> 1);
 					}
 				}
@@ -330,17 +326,14 @@ namespace OpImageMipmapInternal
 								FMemory::Memcpy(&Row1Bits, SrcRow1, sizeof(uint64));
 								
 								const bool bOutOfBounds = Row0Offset.X + 1 >= BatchDecSizeInPixels.X;
-								
-								constexpr uint64 ShiftMask = 0xFEFEFEFEFEFEFEFE;
-								
-								const uint64 XorRow0Row1Bits = Row0Bits ^ Row1Bits;
-								const uint64 ErrorCorrection = XorRow0Row1Bits & 0x0101010101010101;
-								
-								// Average of 2 unsigned integers without overflow extended to work on multiple bytes.
-								const uint64 AvgLowBits = (Row0Bits & Row1Bits) + ((XorRow0Row1Bits & ShiftMask) >> 1) + ErrorCorrection;
+									
+								// Average of 2 unsigned integers with rounding without overflow extended to work on multiple bytes.
+								constexpr uint64 ShiftMask = 0x7F7F7F7F7F7F7F7F;
+						
+								const uint64 AvgLowBits = (Row0Bits | Row1Bits) - (((Row0Bits ^ Row1Bits) >> 1) & ShiftMask);
 								const uint64 AvgHighBits = bOutOfBounds ? AvgLowBits : (AvgLowBits >> NumChannels*8);
-								const uint32 Result = (AvgLowBits & AvgHighBits) + (((AvgLowBits ^ AvgHighBits) & ShiftMask) >> 1);
-								
+								const uint32 Result = (AvgLowBits | AvgHighBits) - (((AvgLowBits ^ AvgHighBits) >> 1) & ShiftMask);
+
 								FMemory::Memcpy(DestPixel, &Result, NumChannels);
 							}
 							else // constexpr Filter == EMipmapFilterType::Unfiltered
