Fix for automatic compilation in editor not picking up on changes from in a Customizable Object Macro Library. Now changes will be picked up on autocompilation and be represented in the output.

#Synced-CL 52279582

diff --git a/Engine/Plugins/Mutable/Source/CustomizableObjectEditor/Private/MuCOE/CustomizableObjectEditorModule.cpp b/Engine/Plugins/Mutable/Source/CustomizableObjectEditor/Private/MuCOE/CustomizableObjectEditorModule.cpp
--- a/Engine/Plugins/Mutable/Source/CustomizableObjectEditor/Private/MuCOE/CustomizableObjectEditorModule.cpp
+++ b/Engine/Plugins/Mutable/Source/CustomizableObjectEditor/Private/MuCOE/CustomizableObjectEditorModule.cpp
@@ -1113,6 +1113,11 @@ TMap<FName, FGuid> FCustomizableObjectEditorModule::GetParticipatingObjects(cons
 				}
 			}
 		}
+		
+		else if (UCustomizableObjectNodeMacroInstance* MacroInstanceNode = Cast<UCustomizableObjectNodeMacroInstance>(&Node))
+		{
+			AddObject(MacroInstanceNode->ParentMacroLibrary);
+		}
 	};
 
 	TArray<const UCustomizableObject*> VisitedObjects;
