[Mutable] More robust Start Update checks.

- Check for reachability.
- Removed redundant checks.

#Synced-CL 52295366

diff --git a/Engine/Plugins/Mutable/Source/CustomizableObject/Internal/MuCO/CustomizableObjectInstancePrivate.h b/Engine/Plugins/Mutable/Source/CustomizableObject/Internal/MuCO/CustomizableObjectInstancePrivate.h
--- a/Engine/Plugins/Mutable/Source/CustomizableObject/Internal/MuCO/CustomizableObjectInstancePrivate.h
+++ b/Engine/Plugins/Mutable/Source/CustomizableObject/Internal/MuCO/CustomizableObjectInstancePrivate.h
@@ -241,10 +241,7 @@ public:
 	UE_API FCustomizableObjectInstanceDescriptor& GetDescriptor() const;
 
 	UE_API void CopyParametersFromInstance(UCustomizableObjectInstance* Instance);
-
-	/** Return true if the instance is not locked and if it's compiled. */
-	UE_API bool CanUpdateInstance() const;
-
+	
 	/** Finds in IntParameters a parameter with name ParamName, returns the index if found, -1 otherwise. */
 	UE_API int32 FindIntParameterNameIndex(const FString& ParamName) const;
 
diff --git a/Engine/Plugins/Mutable/Source/CustomizableObject/Internal/MuCO/CustomizableObjectSystemPrivate.h b/Engine/Plugins/Mutable/Source/CustomizableObject/Internal/MuCO/CustomizableObjectSystemPrivate.h
--- a/Engine/Plugins/Mutable/Source/CustomizableObject/Internal/MuCO/CustomizableObjectSystemPrivate.h
+++ b/Engine/Plugins/Mutable/Source/CustomizableObject/Internal/MuCO/CustomizableObjectSystemPrivate.h
@@ -239,12 +239,8 @@ struct FComponentDescriptor
 class FUpdateContextPrivate
 {
 public:
-	FUpdateContextPrivate(UCustomizableObjectInstance& InInstance, const FCustomizableObjectInstanceDescriptor& Descriptor);
-
-	CUSTOMIZABLEOBJECT_API FUpdateContextPrivate(UCustomizableObjectInstance& InInstance);
+	void Init(bool bUseCommitedDescriptor);
 	
-	bool IsContextValid() const;
-
 	const FCustomizableObjectInstanceDescriptor& GetCapturedDescriptor() const;
 
 	const FDescriptorHash& GetCapturedDescriptorHash() const;
@@ -383,18 +379,12 @@ public:
 
 	/** Add to the high priority queue or not */
 	uint8 bIsHighPriority : 1 = false;
-
-	/** The Context has been successfully created. */
-	uint8 bValid : 1 = false;
-
+	
 	/** Do not optimize the update. */
 	uint8 bForce : 1= false;
 
 	/** Update due to a bake. */
 	uint8 bBake : 1 = false;
-
-	/** User config. */
-	uint8 bProgressiveMipStreamingEnabled : 1 = false;
 	
 	uint8 bIsUpdateAborted : 1 = false;
 		
@@ -496,11 +486,7 @@ public:
 	static int32 MutableMeshesLODBias;
 	static int32 MaxTextureSizeToGenerate;
 	static bool bGenerateInstancesWithinRange;
-
-#if WITH_EDITOR
-	UE::Mutable::Private::FImageOperator::FImagePixelFormatFunc ImageFormatOverrideFunc;
-#endif
-
+	
 	// IStreamingManager interface
 	virtual void UpdateResourceStreaming(float DeltaTime, bool bProcessEverything = false) override;
 	virtual int32 BlockTillAllRequestsFinished(float TimeLimit = 0.0f, bool bLogResults = false) override;
@@ -520,12 +506,7 @@ public:
 	UCustomizableObjectSystem* GetPublic() const;
 
 	CUSTOMIZABLEOBJECT_API void EnqueueUpdateSkeletalMesh(const TSharedRef<FUpdateContextPrivate>& Context);
-
-	/** Updated at the beginning of each tick. */
-	int32 GetNumSkeletalMeshes() const;
-	UPROPERTY(Transient)
-	TObjectPtr<UCustomizableObjectInstance> CurrentInstanceBeingUpdated = nullptr;
-
+	
 	TSharedPtr<FUpdateContextPrivate> CurrentMutableOperation = nullptr;
 
 	// Handle to the registered TickDelegate.
diff --git a/Engine/Plugins/Mutable/Source/CustomizableObject/Private/MuCO/CustomizableObjectInstance.cpp b/Engine/Plugins/Mutable/Source/CustomizableObject/Private/MuCO/CustomizableObjectInstance.cpp
--- a/Engine/Plugins/Mutable/Source/CustomizableObject/Private/MuCO/CustomizableObjectInstance.cpp
+++ b/Engine/Plugins/Mutable/Source/CustomizableObject/Private/MuCO/CustomizableObjectInstance.cpp
@@ -754,38 +754,12 @@ void UCustomizableInstancePrivate::PostEditChangePropertyWithoutEditor()
 }
 
 
-bool UCustomizableInstancePrivate::CanUpdateInstance() const
-{
-	UCustomizableObject* Object = GetPublic()->GetCustomizableObject();
-	if (!Object)
-	{
-		return false;
-	}
-
-#if WITH_EDITOR
-	if (Object->GetPrivate()->IsLocked())
-	{
-		return false;
-	}
-
-	if (!Object->IsCompiled())
-	{
-		return false;
-	}
-
-	return true;
-	
-#else
-	return Object->IsCompiled();
-#endif
-}
-
-
 void UCustomizableObjectInstance::UpdateSkeletalMeshAsync(bool, bool bForceHighPriority)
 {
 	UCustomizableObjectSystemPrivate* SystemPrivate = UCustomizableObjectSystem::GetInstance()->GetPrivate();
 
-	const TSharedRef<FUpdateContextPrivate> Context = MakeShared<FUpdateContextPrivate>(*this);
+	const TSharedRef<FUpdateContextPrivate> Context = MakeShared<FUpdateContextPrivate>();
+	Context->Instance = this;
 	Context->bIsHighPriority = bForceHighPriority;
 	
 	SystemPrivate->EnqueueUpdateSkeletalMesh(Context);
@@ -796,7 +770,8 @@ void UCustomizableObjectInstance::UpdateSkeletalMeshAsyncResult(FInstanceUpdateD
 {
 	UCustomizableObjectSystemPrivate* SystemPrivate = UCustomizableObjectSystem::GetInstance()->GetPrivate();
 
-	const TSharedRef<FUpdateContextPrivate> Context = MakeShared<FUpdateContextPrivate>(*this);
+	const TSharedRef<FUpdateContextPrivate> Context = MakeShared<FUpdateContextPrivate>();
+	Context->Instance = this;
 	Context->bIsHighPriority = bForceHighPriority;
 	Context->UpdateCallback = Callback;
 	
@@ -808,7 +783,8 @@ void UCustomizableObjectInstance::UpdateSkeletalMeshAsyncResult(FInstanceUpdateN
 {
 	UCustomizableObjectSystemPrivate* SystemPrivate = UCustomizableObjectSystem::GetInstance()->GetPrivate();
 
-	const TSharedRef<FUpdateContextPrivate> Context = MakeShared<FUpdateContextPrivate>(*this);
+	const TSharedRef<FUpdateContextPrivate> Context = MakeShared<FUpdateContextPrivate>();
+	Context->Instance = this;
 	Context->bIsHighPriority = bForceHighPriority;
 	Context->UpdateNativeCallback = Callback;
 	
@@ -2747,18 +2723,7 @@ bool MutableTextureUsesOfflineProcessedData()
 
 void SetTexturePropertiesFromMutableImageProps(UTexture2D* Texture, const FMutableModelImageProperties& Props, bool bNeverStream)
 {
-#if !PLATFORM_DESKTOP
-	if (UCustomizableObjectSystem::GetInstance()->GetPrivate()->EnableMutableProgressiveMipStreaming <= 0)
-	{
-		Texture->NeverStream = true;
-	}
-	else
-	{
-		Texture->NeverStream = bNeverStream;
-	}
-#else
 	Texture->NeverStream = bNeverStream;
-#endif
 	Texture->bNotOfflineProcessed = !MutableTextureUsesOfflineProcessedData();
 
 	Texture->SRGB = Props.SRGB;
diff --git a/Engine/Plugins/Mutable/Source/CustomizableObject/Private/MuCO/CustomizableObjectMipDataProvider.cpp b/Engine/Plugins/Mutable/Source/CustomizableObject/Private/MuCO/CustomizableObjectMipDataProvider.cpp
--- a/Engine/Plugins/Mutable/Source/CustomizableObject/Private/MuCO/CustomizableObjectMipDataProvider.cpp
+++ b/Engine/Plugins/Mutable/Source/CustomizableObject/Private/MuCO/CustomizableObjectMipDataProvider.cpp
@@ -308,8 +308,6 @@ int32 FMutableTextureMipDataProvider::GetMips(const FTextureUpdateContext& Conte
 		}
 		else // Generate a Mip Request to Mutable
 		{
-			check(UCustomizableObjectSystem::GetInstance()->GetPrivate()->EnableMutableProgressiveMipStreaming == 1);
-
 			if (!OperationData.IsValid())
 			{
 				OperationData = MakeShared<FMutableImageOperationData>();
diff --git a/Engine/Plugins/Mutable/Source/CustomizableObject/Private/MuCO/CustomizableObjectSystem.cpp b/Engine/Plugins/Mutable/Source/CustomizableObject/Private/MuCO/CustomizableObjectSystem.cpp
--- a/Engine/Plugins/Mutable/Source/CustomizableObject/Private/MuCO/CustomizableObjectSystem.cpp
+++ b/Engine/Plugins/Mutable/Source/CustomizableObject/Private/MuCO/CustomizableObjectSystem.cpp
@@ -207,29 +207,61 @@ void RemoveUnusedOnScreenMessages(TMap<TWeakObjectPtr<const UCustomizableObject>
 }
 
 
-FUpdateContextPrivate::FUpdateContextPrivate(UCustomizableObjectInstance& InInstance, const FCustomizableObjectInstanceDescriptor& Descriptor)
+bool CanUpdate(UCustomizableObjectInstance* Instance)
 {
-	check(IsInGameThread());
-
-	if (!IsValid(&InInstance))
+	if (!IsValid(Instance))
 	{
-		return;
+		return false;
 	}
 	
-	Instance = &InInstance;
+	UCustomizableObject* Object = Instance->GetCustomizableObject();
+	if (!IsValid(Object))
+	{
+		return false;
+	}
 
-	UCustomizableObject* InObject = InInstance.GetCustomizableObject();
-	if (!IsValid(InObject))
+	if (!Object->IsCompiled())
 	{
-		return;
+		return false;
+	}
+	
+#if WITH_EDITOR
+	if (Object->GetPrivate()->IsLocked())
+	{
+		return false;
 	}
+#endif
+	
+	return true;
+}
 
-	Object = InObject;
 
-	CapturedDescriptor = Descriptor;
-	InObject->GetPrivate()->ApplyStateForcedValuesToParameters(CapturedDescriptor);
+void FUpdateContextPrivate::Init(bool bUseCommitedDescriptor)
+{
+	check(IsInGameThread());
+	
+	UCustomizableObjectInstance* StrongInstance = Instance.Get();
+	check(IsValid(StrongInstance));
+	
+	UCustomizableObject* StrongObject = StrongInstance->GetCustomizableObject();
+	check(IsValid(StrongObject));
+
+	Object = StrongObject;
 
-	if (const UModelResources* ModelResources = InObject->GetPrivate()->GetModelResources())
+	const FDescriptorHash& CommittedHash = StrongInstance->GetPrivate()->CommittedDescriptorHash;
+	if (bUseCommitedDescriptor && CommittedHash != FDescriptorHash())
+	{
+		CapturedDescriptor = StrongInstance->GetPrivate()->CommittedDescriptor;
+		CapturedDescriptorHash = CommittedHash;
+	}
+	else
+	{
+		CapturedDescriptor = StrongInstance->GetPrivate()->GetDescriptor();
+		StrongObject->GetPrivate()->ApplyStateForcedValuesToParameters(CapturedDescriptor);
+		CapturedDescriptorHash = FDescriptorHash(CapturedDescriptor);
+	}
+
+	if (const UModelResources* ModelResources = StrongObject->GetPrivate()->GetModelResources())
 	{
 		ComponentNames = ModelResources->ComponentNamesPerObjectComponent;
 
@@ -245,8 +277,6 @@ FUpdateContextPrivate::FUpdateContextPrivate(UCustomizableObjectInstance& InInst
 			FirstResidentLOD.Add(ComponentName, Pair.Value);
 		}
 	}
-
-	CapturedDescriptorHash = FDescriptorHash(Descriptor);
 	
 	UCustomizableObjectSystem* System = UCustomizableObjectSystem::GetInstanceChecked();
 	UCustomizableObjectSystemPrivate* SystemPrivate = System->GetPrivate();
@@ -254,15 +284,7 @@ FUpdateContextPrivate::FUpdateContextPrivate(UCustomizableObjectInstance& InInst
 	MutableSystem = SystemPrivate->MutableSystem;	
 	check(MutableSystem);
 	
-	bProgressiveMipStreamingEnabled = System->IsProgressiveMipStreamingEnabled();
-
-#if WITH_EDITOR
-	PixelFormatOverride = SystemPrivate->ImageFormatOverrideFunc;
-#endif
-
-	bKeepOwnershipOfGeneratedResources = InInstance.GetPrivate()->bKeepOwnershipOfGeneratedResources;
-	
-	bValid = true;
+	bKeepOwnershipOfGeneratedResources = StrongInstance->GetPrivate()->bKeepOwnershipOfGeneratedResources;
 }
 
 
@@ -474,18 +496,6 @@ bool TickWarnings(float DeltaTime)
 }
 
 
-bool FUpdateContextPrivate::IsContextValid() const
-{
-	return bValid;
-}
-
-
-FUpdateContextPrivate::FUpdateContextPrivate(UCustomizableObjectInstance& InInstance) :
-	FUpdateContextPrivate(InInstance, InInstance.GetPrivate()->GetDescriptor())
-{
-}
-
-
 const FCustomizableObjectInstanceDescriptor& FUpdateContextPrivate::GetCapturedDescriptor() const
 {
 	return CapturedDescriptor;
@@ -596,7 +606,6 @@ void UCustomizableObjectSystem::InitSystem()
 	check(Private != nullptr);
 
 	Private->CurrentMutableOperation = nullptr;
-	Private->CurrentInstanceBeingUpdated = nullptr;
 
 	Private->LastWorkingMemoryBytes = CVarWorkingMemoryKB->GetInt() * 1024;
 
@@ -642,7 +651,6 @@ void UCustomizableObjectSystem::BeginDestroy()
 
 		// Clear the ongoing operation
 		Private->CurrentMutableOperation = nullptr;
-		Private->CurrentInstanceBeingUpdated = nullptr;
 
 		UCustomizableObjectSystemPrivate::SSystem = nullptr;
 
@@ -743,6 +751,12 @@ static FAutoConsoleVariable CVarDescriptorDebugPrint(
 	ECVF_Default);
 
 
+namespace impl
+{
+	void Task_Mutable_EndUpdate(TSharedPtr<UE::Mutable::Private::FLiveInstance> LiveInstance, TSharedPtr<UE::Mutable::Private::FSystem> MutableSystem);
+}
+
+
 void FinishUpdateGlobal(const TSharedRef<FUpdateContextPrivate>& Context)
 {
 	MUTABLE_CPUPROFILER_SCOPE(FinishUpdateGlobal);
@@ -757,8 +771,14 @@ void FinishUpdateGlobal(const TSharedRef<FUpdateContextPrivate>& Context)
 	if (System &&
 		Context->UpdateStarted)
 	{
-		SystemPrivate->CurrentInstanceBeingUpdated = nullptr;
 		SystemPrivate->CurrentMutableOperation = nullptr;
+		
+		SystemPrivate->LastUpdateMutableTask = SystemPrivate->MutableTaskGraph.AddMutableThreadTask(
+			TEXT("Task_Mutable_EndUpdate"),
+			[LiveInstance = Context->LiveInstance, MutableSystem = SystemPrivate->MutableSystem]()
+			{
+				impl::Task_Mutable_EndUpdate(LiveInstance, MutableSystem);
+			});
 	}
 	
 	if (IsValid(Instance))
@@ -822,12 +842,14 @@ void FinishUpdateGlobal(const TSharedRef<FUpdateContextPrivate>& Context)
 				{
 					continue;
 				}
+				
 #if WITH_EDITOR
 				if (InstanceUsage->GetPrivate()->bIsNetModeDedicatedServer)
 				{
 					continue;
 				}
 #endif
+				
 				InstanceUsage->GetPrivate()->Callbacks();
 			}
 		}
@@ -1262,41 +1284,20 @@ void UCustomizableObjectSystemPrivate::DiscardInstance(UCustomizableObjectInstan
 }
 
 
-void UCustomizableObjectSystemPrivate::EnqueueUpdateSkeletalMesh(const TSharedRef<FUpdateContextPrivate>& Context)
+void TryAutomaticCompilation(UCustomizableObjectInstance* Instance)
 {
-	MUTABLE_CPUPROFILER_SCOPE(FCustomizableObjectSystemPrivate::EnqueueUpdateSkeletalMesh);
-	check(IsInGameThread());
-	
-	UCustomizableObjectInstance* Instance = Context->Instance.Get();
-	UCustomizableObject* Object = Context->Object.Get();
-
-	if (!Instance || !Object)
+#if WITH_EDITOR
+	if (!IsValid(Instance))
 	{
-		Context->UpdateResult = EUpdateResult::Error;
-		FinishUpdateGlobal(Context);
 		return;
 	}
 
-	if (!Context->IsContextValid()) 
-	{
-		Context->UpdateResult = EUpdateResult::Error;
-		FinishUpdateGlobal(Context);
-		return;
-	}
-	
-	UE_LOGF(LogMutable, Verbose, "Enqueue Update Skeletal Mesh Async. CustomizableObject=%ls Instance=%ls, Frame=%d", *Object->GetFName().ToString(), *Instance->GetFName().ToString(), GFrameNumber);
-	UCustomizableInstancePrivate* InstancePrivate = Instance->GetPrivate();
-		
-	if (!bIsMutableEnabled)
+	UCustomizableObject* Object = Instance->GetCustomizableObject();
+	if (!IsValid(Object))
 	{
-		// Mutable is disabled. Set the reference SkeletalMesh and finish the update with success to avoid breaking too many things.
-		Context->UpdateResult = EUpdateResult::Success;
-		InstancePrivate->ForceSetReferenceSkeletalMesh();
-		FinishUpdateGlobal(Context);
 		return;
 	}
 
-#if WITH_EDITOR
 	if (!Object->IsCompiled() &&
 		Object->GetPrivate()->CompilationResult != ECompilationResultPrivate::Errors) // Avoid constantly retry failed compilations.
 	{
@@ -1307,20 +1308,44 @@ void UCustomizableObjectSystemPrivate::EnqueueUpdateSkeletalMesh(const TSharedRe
 
 			EditorModule->CompileCustomizableObject(*Object, &CompileParams, true, false);
 		}
-
-		Context->UpdateResult = EUpdateResult::ErrorUncompiled;
-		FinishUpdateGlobal(Context);
-		return;
 	}
 #endif
+}
+
+
+void UCustomizableObjectSystemPrivate::EnqueueUpdateSkeletalMesh(const TSharedRef<FUpdateContextPrivate>& Context)
+{
+	MUTABLE_CPUPROFILER_SCOPE(FCustomizableObjectSystemPrivate::EnqueueUpdateSkeletalMesh);
+	check(IsInGameThread());
 	
-	if (!InstancePrivate->CanUpdateInstance())
+	UCustomizableObjectInstance* Instance = Context->Instance.Get();
+
+	TryAutomaticCompilation(Instance);
+
+	if (!CanUpdate(Instance))
 	{
 		Context->UpdateResult = EUpdateResult::Error;
 		FinishUpdateGlobal(Context);
 		return;
 	}
+	
+	// TODO GMT Move this further down
+	Context->Init(false);
 
+	UCustomizableObject* Object = Context->Object.Get();
+	UCustomizableInstancePrivate* InstancePrivate = Instance->GetPrivate();
+
+	UE_LOGF(LogMutable, Verbose, "Enqueue Update Skeletal Mesh Async. CustomizableObject=%ls Instance=%ls, Frame=%d", *Object->GetFName().ToString(), *Instance->GetFName().ToString(), GFrameNumber);
+		
+	if (!bIsMutableEnabled)
+	{
+		// Mutable is disabled. Set the reference SkeletalMesh and finish the update with success to avoid breaking too many things.
+		Context->UpdateResult = EUpdateResult::Success;
+		InstancePrivate->ForceSetReferenceSkeletalMesh();
+		FinishUpdateGlobal(Context);
+		return;
+	}
+	
 	if (!Context->bForce)
 	{
 		if (CurrentMutableOperation &&
@@ -1351,13 +1376,9 @@ void UCustomizableObjectSystemPrivate::EnqueueUpdateSkeletalMesh(const TSharedRe
 			}
 			else
 			{
-				Context->bOptimizedUpdate = false;
-
 				// The user may have changed the AttachParent and we need to re-customize it.
 				// In case nothing need to be re-customized, the update will be considered ErrorOptimized.
 				UpdateSkeletalMesh(Context);
-				Context->UpdateResult = EUpdateResult::Success;
-
 				FinishUpdateGlobal(Context);
 				return;
 			}
@@ -1461,8 +1482,8 @@ bool UCustomizableObjectSystem::LockObject(UCustomizableObject* InObject)
 	if (Private)
 	{
 		// If the current instance is for this object, make the lock fail by returning false
-		if (Private->CurrentInstanceBeingUpdated &&
-			Private->CurrentInstanceBeingUpdated->GetCustomizableObject() == InObject)
+		if (Private->CurrentMutableOperation &&
+			Private->CurrentMutableOperation->Object.Get() == InObject)
 		{
 			UE_LOGF(LogMutable, Warning, "---- failed to lock object %ls", *InObject->GetName());
 
@@ -1790,9 +1811,11 @@ namespace impl
 	{
 		MUTABLE_CPUPROFILER_SCOPE(Task_Mutable_GetSurfaces)
 		
-		const UCustomizableObject* CustomizableObject = OperationData->Object.Get();
-		const UModelResources& ModelResources = *CustomizableObject->GetPrivate()->GetModelResources();
-
+		if (OperationData->bIsUpdateAborted)
+		{
+			return;
+		}
+		
 		const UE::Mutable::Private::FInstance* MutableInstance = OperationData->MutableInstance.Get();
 
 		TArray<uint32> GeneratedSurfaceIds;
@@ -1802,7 +1825,7 @@ namespace impl
 		for (int32 ComponentIndex = 0; ComponentIndex < OperationData->InstanceUpdateData.Components.Num(); ++ComponentIndex)
 		{
 			TSharedRef<FInstanceUpdateData::FComponent>& Component = OperationData->InstanceUpdateData.Components[ComponentIndex];
-			const FName& ComponentName = ModelResources.ComponentNamesPerObjectComponent[Component->Id];
+			const FName& ComponentName = OperationData->ComponentNames[Component->Id];
 
 			for (int32 MutableLODIndex = OperationData->FirstLODAvailable[ComponentName]; MutableLODIndex < Component->LODCount; ++MutableLODIndex)
 			{
@@ -1864,9 +1887,27 @@ namespace impl
 	{
 		MUTABLE_CPUPROFILER_SCOPE(Task_Mutable_GetMaterials)
 
-		const UCustomizableObject* CustomizableObject = Context->Object.Get();
-		const UModelResources* ModelResources = CustomizableObject->GetPrivate()->GetModelResources();
+		if (Context->bIsUpdateAborted)
+		{
+			return;
+		}
+		
+		// TODO We should not access UObjects outside GT
+		const TStrongObjectPtr<UCustomizableObject> CustomizableObject = Context->Object.Pin();
+		if (!IsValid(CustomizableObject.Get()))
+		{
+			Context->bIsUpdateAborted = true;
+			return;
+		}
 
+		// TODO We should not access UObjects outside GT
+		const TStrongObjectPtr<UModelResources> ModelResources = TStrongObjectPtr(CustomizableObject->GetPrivate()->GetModelResources());
+		if (!IsValid(ModelResources.Get()))
+		{
+			Context->bIsUpdateAborted = true;
+			return;
+		}
+		
 		UE::Tasks::FTask LastTask = UE::Tasks::MakeCompletedTask<void>();
 		
 		for (int32 MaterialIndex = 0; MaterialIndex < Context->InstanceUpdateData.Materials.Num(); ++MaterialIndex)
@@ -1877,7 +1918,7 @@ namespace impl
 			UE::Tasks::AddNested(GetMaterialTask);
 			LastTask = GetMaterialTask;
 			
-			LastTask = UE::Tasks::Launch(TEXT("GetMaterialResult"), [Context, ModelResources = TStrongObjectPtr(ModelResources), Material, GetMaterialTask]() mutable
+			LastTask = UE::Tasks::Launch(TEXT("GetMaterialResult"), [Context, ModelResources, Material, GetMaterialTask]() mutable
 			{
 				Material->Material = GetMaterialTask.GetResult();
 
@@ -1916,13 +1957,15 @@ namespace impl
 	{
 		MUTABLE_CPUPROFILER_SCOPE(Task_Mutable_GetMeshes)
 
+		if (OperationData->bIsUpdateAborted)
+		{
+			return;
+		}
+		
 		OperationData->TaskGetMeshesTimeStart = FPlatformTime::Seconds();
 
 		check(OperationData->Parameters);
 		OperationData->InstanceUpdateData.Clear();
-
-		check(UCustomizableObjectSystem::GetInstance() != nullptr);
-		check(UCustomizableObjectSystem::GetInstance()->GetPrivate() != nullptr);
 		
 		check(OperationData->LiveInstance);
 		
@@ -2107,7 +2150,12 @@ namespace impl
 	void Task_Mutable_GetImageDescriptors(const TSharedRef<FUpdateContextPrivate>& OperationData)
 	{
 		MUTABLE_CPUPROFILER_SCOPE(Task_Mutable_GetImageDescriptros)
-
+		
+		if (OperationData->bIsUpdateAborted)
+		{
+			return;
+		}
+		
 		OperationData->TaskGetImagesTimeStart = FPlatformTime::Seconds();		
 
 		const TSharedPtr<TArray<UE::Mutable::Private::FImageId>> ImagesInThisInstance = MakeShared<TArray<UE::Mutable::Private::FImageId>>();
@@ -2165,6 +2213,11 @@ namespace impl
 	{
 		MUTABLE_CPUPROFILER_SCOPE(Task_Mutable_GetImages)
 		
+		if (OperationData->bIsUpdateAborted)
+		{
+			return;
+		}
+		
 		const TSharedPtr<TArray<UE::Mutable::Private::FImageId>> RequestedGetImages = MakeShared<TArray<UE::Mutable::Private::FImageId>>();
 
 		UE::Tasks::FTask LastTask;
@@ -2285,6 +2338,11 @@ namespace impl
 	{
 		MUTABLE_CPUPROFILER_SCOPE(Subtask_Mutable_PrepareTextures)
 
+		if (OperationData->bIsUpdateAborted)
+		{
+			return;
+		}
+		
 		for (const TSharedRef<FInstanceUpdateData::FMaterial>& Material : OperationData->InstanceUpdateData.Materials)
 		{
 			for (int32 ImageIndex = 0; ImageIndex < Material->ImageCount; ++ImageIndex)
@@ -2315,19 +2373,20 @@ namespace impl
 	
 
 	// This runs in a worker thread.
-	void Task_Mutable_EndUpdate(TSharedRef<UE::Mutable::Private::FLiveInstance> LiveInstance, TSharedPtr<UE::Mutable::Private::FSystem> MutableSystem, bool bLiveUpdateMode)
+	void Task_Mutable_EndUpdate(TSharedPtr<UE::Mutable::Private::FLiveInstance> LiveInstance, TSharedPtr<UE::Mutable::Private::FSystem> MutableSystem)
 	{
 		MUTABLE_CPUPROFILER_SCOPE(Task_Mutable_ReleaseInstance)
-
-		check(MutableSystem);
-
-		const bool bClearRoms = UCustomizableObjectSystem::ShouldClearWorkingMemoryOnUpdateEnd();
-		const bool bFreeCache = !CVarReuseLiveInstanceCacheBetweenUpdates.GetValueOnAnyThread();
-		MutableSystem->EndUpdate(LiveInstance, bClearRoms, bFreeCache);
+		
+		if (LiveInstance && MutableSystem)
+		{
+			const bool bClearRoms = UCustomizableObjectSystem::ShouldClearWorkingMemoryOnUpdateEnd();
+			const bool bFreeCache = !CVarReuseLiveInstanceCacheBetweenUpdates.GetValueOnAnyThread();
+			MutableSystem->EndUpdate(LiveInstance.ToSharedRef(), bClearRoms, bFreeCache);
+		}
 		
 		UCustomizableObjectSystem::GetInstance()->GetPrivate()->MutableTaskGraph.AllowLaunchingMutableTaskLowPriority(true, true);
 	}
-
+	
 	
 	void Task_Game_Callbacks(const TSharedRef<FUpdateContextPrivate>& Context)
 	{
@@ -2336,6 +2395,13 @@ namespace impl
 
 		check(IsInGameThread());
 
+		if (Context->bIsUpdateAborted)
+		{
+			Context->UpdateResult = EUpdateResult::Error;
+			FinishUpdateGlobal(Context);
+			return;
+		}
+		
 		UCustomizableObjectSystem* System = UCustomizableObjectSystem::GetInstance();
 		if (!IsValid(System))
 		{
@@ -2378,6 +2444,13 @@ namespace impl
 
 		check(IsInGameThread());
 
+		if (OperationData->bIsUpdateAborted)
+		{
+			OperationData->UpdateResult = EUpdateResult::Error;
+			FinishUpdateGlobal(OperationData);
+			return;
+		}
+		
 		UCustomizableObjectSystem* System = UCustomizableObjectSystem::GetInstance();
 		if (!IsValid(System))
 		{
@@ -2386,108 +2459,95 @@ namespace impl
 			return;
 		}
 
-		if (OperationData->bIsUpdateAborted)
+		UCustomizableObjectSystemPrivate* SystemPrivate = System->GetPrivate();
+
+		UCustomizableObject* Object = OperationData->Object.Get();
+		if (!IsValid(Object))
 		{
 			OperationData->UpdateResult = EUpdateResult::Error;
 			FinishUpdateGlobal(OperationData);
 			return;
 		}
 		
+		UCustomizableObjectInstance* Instance = OperationData->Instance.Get();
+        if (!IsValid(Instance))
+        {
+        	OperationData->UpdateResult = EUpdateResult::Error;
+        	FinishUpdateGlobal(OperationData);
+        	return;
+        }
+		
+		UCustomizableInstancePrivate* InstancePrivate = Instance->GetPrivate();
+
+		
 		if (CVarEnableRealTimeMorphTargets.GetValueOnAnyThread())
 		{
 			// TODO: This subtask should execute before Convert resources in a worker thread but after 
 			// Loading resources. For now keep it here.
 			Subtask_Mutable_PrepareRealTimeMorphData(OperationData);
 		}
-
-		UCustomizableObjectInstance* Instance = OperationData->Instance.Get();
-
-		// Actual work
-		const bool bInstanceInvalid = !IsValid(Instance);
-		if (!bInstanceInvalid)
+		
+		// Extension data
 		{
-			UCustomizableInstancePrivate* InstancePrivate = Instance->GetPrivate();
-
-			// Extension data
+			InstancePrivate->ExtensionData.Empty(OperationData->ExtensionData.Num());
+			for (UE::Mutable::Private::TManagedPtr<const UE::Mutable::Private::FExtensionData>& Data : OperationData->ExtensionData)
 			{
-				InstancePrivate->ExtensionData.Empty(OperationData->ExtensionData.Num());
-				for (UE::Mutable::Private::TManagedPtr<const UE::Mutable::Private::FExtensionData>& Data : OperationData->ExtensionData)
+				if (UObject* PassthroughObject = Data->PassthroughObject.Get())
 				{
-					if (UObject* PassthroughObject = Data->PassthroughObject.Get())
-					{
-						InstancePrivate->ExtensionData.Add(PassthroughObject);
-					}
+					InstancePrivate->ExtensionData.Add(PassthroughObject);
 				}
 			}
-			
-			// Prepare the component data
-			//-------------------------------------------------------------
-			{
-				UCustomizableObject* Object = OperationData->Object.Get();
-				check(Object);
+		}
 		
-				const UModelResources& ModelResources = Object->GetPrivate()->GetModelResourcesChecked();
-				
-				const int32 NumComponents = OperationData->MutableInstance->GetComponentCount();
-				InstancePrivate->ComponentsData.SetNum(NumComponents);
-
-				TArray<FName> ComponentNames;
-				ComponentNames.Reserve(NumComponents);
-
-				for (int32 ComponentIndex = 0; ComponentIndex < NumComponents; ++ComponentIndex)
-				{
-					const UE::Mutable::Private::FComponentId ComponentId = OperationData->MutableInstance->GetComponentId(ComponentIndex);
-					const FName& ComponentName = ModelResources.ComponentNamesPerObjectComponent[ComponentId];
-
-					if (ComponentNames.Contains(ComponentName)) // Components with the same name. Abort.
-					{
-						UE_LOGF(LogMutable, Error, "Update error. Trying to generate multiple components with the same name: %ls",
-							*ComponentName.ToString());
-						
-						OperationData->UpdateResult = EUpdateResult::Error;
-						FinishUpdateGlobal(OperationData);
-						return;
-					}
-
-					ComponentNames.Add(ComponentName);
-
-					InstancePrivate->ComponentsData[ComponentIndex].ComponentId = ComponentId;
-					InstancePrivate->ComponentsData[ComponentIndex].ComponentName = ComponentName;
-				}
-			}
+		// Prepare the component data
+		//-------------------------------------------------------------
+		{		
+			const UModelResources& ModelResources = Object->GetPrivate()->GetModelResourcesChecked();
 			
-			// Convert Step
-			//-------------------------------------------------------------
+			const int32 NumComponents = OperationData->MutableInstance->GetComponentCount();
+			InstancePrivate->ComponentsData.SetNum(NumComponents);
 
-			// \TODO: Bring that code here instead of keeping it in the UCustomizableObjectInstance
-			if (InstancePrivate->UpdateSkeletalMesh_PostBeginUpdate0(Instance, OperationData))
-			{
-				// This used to be CustomizableObjectInstance::UpdateSkeletalMesh_PostBeginUpdate1
-				{
-					MUTABLE_CPUPROFILER_SCOPE(UpdateSkeletalMesh_PostBeginUpdate1);
+			TArray<FName> ComponentNames;
+			ComponentNames.Reserve(NumComponents);
 
-					// \TODO: Bring here
-					InstancePrivate->BuildResources(OperationData, Instance);
-				}
+			for (int32 ComponentIndex = 0; ComponentIndex < NumComponents; ++ComponentIndex)
+			{
+				const UE::Mutable::Private::FComponentId ComponentId = OperationData->MutableInstance->GetComponentId(ComponentIndex);
+				const FName& ComponentName = ModelResources.ComponentNamesPerObjectComponent[ComponentId];
 
-				// This used to be CustomizableObjectInstance::UpdateSkeletalMesh_PostBeginUpdate2
+				if (ComponentNames.Contains(ComponentName)) // Components with the same name. Abort.
 				{
-					MUTABLE_CPUPROFILER_SCOPE(UpdateSkeletalMesh_PostBeginUpdate2);
+					UE_LOGF(LogMutable, Error, "Update error. Trying to generate multiple components with the same name: %ls",
+						*ComponentName.ToString());
 					
-#if WITH_EDITORONLY_DATA
-					InstancePrivate->RegenerateImportedModels(OperationData);
-#endif
-					InstancePrivate->PostEditChangePropertyWithoutEditor();
+					OperationData->UpdateResult = EUpdateResult::Error;
+					FinishUpdateGlobal(OperationData);
+					return;
 				}
-			}
-		} // if (!bInstanceValid)
 
-		if (OperationData->UpdateResult == EUpdateResult::Error)
+				ComponentNames.Add(ComponentName);
+
+				InstancePrivate->ComponentsData[ComponentIndex].ComponentId = ComponentId;
+				InstancePrivate->ComponentsData[ComponentIndex].ComponentName = ComponentName;
+			}
+		}
+		
+		const bool bSuccess = InstancePrivate->UpdateSkeletalMesh_PostBeginUpdate0(Instance, OperationData);
+		if (!bSuccess)
 		{
+			OperationData->UpdateResult = EUpdateResult::Error;
 			FinishUpdateGlobal(OperationData);
 			return;
 		}
 		
+		InstancePrivate->BuildResources(OperationData, Instance);
+
+#if WITH_EDITORONLY_DATA
+		InstancePrivate->RegenerateImportedModels(OperationData);
+#endif
+			
+		InstancePrivate->PostEditChangePropertyWithoutEditor();
+		
 		if (FLogBenchmarkUtil::IsBenchmarkingReportingEnabled())
 		{
 			// Memory used in the context of this the update of mesh
@@ -2496,34 +2556,12 @@ namespace impl
 			OperationData->UpdateEndRealPeakBytes = OperationData->UpdateEndPeakBytes + OperationData->UpdateStartBytes;
 		}
 		
-		UCustomizableObjectSystemPrivate* CustomizableObjectSystemPrivate = System->GetPrivate();
-		
-		// Next Task: Release Mutable. We need this regardless if we cancel or not
-		//-------------------------------------------------------------		
-		const TSharedPtr<UE::Mutable::Private::FSystem> MutableSystem = CustomizableObjectSystemPrivate->MutableSystem;
-		CustomizableObjectSystemPrivate->LastUpdateMutableTask = CustomizableObjectSystemPrivate->MutableTaskGraph.AddMutableThreadTask(
-			TEXT("Task_Mutable_EndUpdate"),
-			[LiveInstance = OperationData->LiveInstance, MutableSystem, bLiveUpdateMode = static_cast<bool>(OperationData->bLiveUpdateMode)]()
-			{
-				Task_Mutable_EndUpdate(LiveInstance.ToSharedRef(), MutableSystem, bLiveUpdateMode);
-			});
-		
-		// We no longer need it.
-		OperationData->LiveInstance.Reset();
-
 		// If they are still in this container it means they have not been used. No longer required. Free some memory now.
 		OperationData->ImageToPlatformDataMap.Empty();
 		
-		if (bInstanceInvalid)
-		{
-			OperationData->UpdateResult = EUpdateResult::Error;
-			FinishUpdateGlobal(OperationData);
-			return;
-		}
-		
 		// Next Task: Callbacks
 		//-------------------------------------------------------------
-		CustomizableObjectSystemPrivate->MutableTaskGraph.AddGameThreadTask(
+		SystemPrivate->MutableTaskGraph.AddGameThreadTask(
 			TEXT("Task_Game_Callbacks"),
 			[OperationData](float RemainingTime)
 			{
@@ -2539,6 +2577,8 @@ namespace impl
 
 		check(IsInGameThread());
 
+		Operation->StartUpdateTime = FPlatformTime::Seconds();
+
 		Operation->UpdateStarted = true;
 		TRACE_BEGIN_REGION(UE_MUTABLE_UPDATE_REGION)
 		
@@ -2549,9 +2589,6 @@ namespace impl
 		}
 
 		UCustomizableObjectSystem* System = UCustomizableObjectSystem::GetInstance();
-		
-		Operation->StartUpdateTime = FPlatformTime::Seconds();
-
 		if (!IsValid(System))
 		{
 			Operation->UpdateResult = EUpdateResult::Error;
@@ -2596,8 +2633,6 @@ namespace impl
 		}
 #endif
 		
-		SystemPrivate->CurrentInstanceBeingUpdated = Instance;
-
 		const UModelResources& ModelResources = Operation->Object->GetPrivate()->GetModelResourcesChecked();
 
 		UCustomizableInstancePrivate* InstancePrivate = Instance->GetPrivate();
@@ -2609,12 +2644,12 @@ namespace impl
 		
 		// Streaming disabled from platform settings or from platform CustomizableObjectSystem properties?
 #if PLATFORM_SUPPORTS_TEXTURE_STREAMING
-		if (!IStreamingManager::Get().IsTextureStreamingEnabled() ||
-			!Operation->bProgressiveMipStreamingEnabled
+		if (!System->IsProgressiveMipStreamingEnabled() ||
+			!IStreamingManager::Get().IsTextureStreamingEnabled() ||
 #if WITH_EDITORONLY_DATA
-			|| ModelResources.bIsTextureStreamingDisabled // Was streaming disabled at object-compilation time? 
+			ModelResources.bIsTextureStreamingDisabled || // Was streaming disabled at object-compilation time? 
 #endif
-			)
+			Operation->bBake)
 		{
 			Operation->bNeverStreamMips = true;
 		}
@@ -2703,6 +2738,8 @@ namespace impl
 			TSharedPtr<UE::Mutable::Private::FLiveInstanceLogger> UpdateInstanceLogger = MakeShared<UE::Mutable::Private::FLiveInstanceLogger>();
 			UpdateInstanceLogger->SetLogger(InstancePrivate->UpdateLogger);
 
+			Operation->Parameters = Operation->GetCapturedDescriptor().GetParameters();
+
 			// Ensure that in the case of running with the live update enabled we have mesh and mip streaming disabled
 			check(!Operation->bLiveUpdateMode || Operation->bNeverStreamMips);
 			check(!Operation->bLiveUpdateMode || Operation->MipsToSkip == 0);
@@ -2718,7 +2755,7 @@ namespace impl
 				Operation->PassthroughObjectLoader,
 				Operation->PixelFormatOverride,
 				UpdateInstanceLogger);
-			
+
 			Operation->LiveInstance->Parameters = Operation->Parameters;
 			Operation->LiveInstance->State = Operation->GetCapturedDescriptor().GetState();
 			
@@ -3308,16 +3345,19 @@ int32 UCustomizableObjectSystem::GetTotalInstances() const
 	return UCustomizableObjectSystemPrivate::GetCustomizableObjectInstances().Num();
 }
 
+
 int64 UCustomizableObjectSystem::GetTextureMemoryUsed() const
 {
 	return GetPrivate()->LogBenchmarkUtil.TextureGPUSize.GetValue();
 }
 
+
 int32 UCustomizableObjectSystem::GetAverageBuildTime() const
 {
 	return GetPrivate()->LogBenchmarkUtil.InstanceBuildTimeAvrg.GetValue() * 1000;
 }
 
+
 PRAGMA_DISABLE_DEPRECATION_WARNINGS
 bool UCustomizableObjectSystem::IsSupport16BitBoneIndexEnabled() const
 {
@@ -3325,6 +3365,7 @@ bool UCustomizableObjectSystem::IsSupport16BitBoneIndexEnabled() const
 }
 PRAGMA_ENABLE_DEPRECATION_WARNINGS
 
+
 bool UCustomizableObjectSystem::IsProgressiveMipStreamingEnabled() const
 {
 	return GetPrivate()->EnableMutableProgressiveMipStreaming != 0;
@@ -3480,14 +3521,9 @@ bool UCustomizableObjectSystemPrivate::TryStartUpdate(const TSharedRef<FUpdateCo
 	MUTABLE_CPUPROFILER_SCOPE(TryStartUpdate);
 
 	check(!CurrentMutableOperation); // Can not start an update if there is already another in progress
-	check(Context->bValid);
-	check(Context->Object.IsValid()) // The CO has to be alive to start the update
-	check(Context->Object->IsCompiled()) // The CO has to be compiled to start the update
-
-	const UCustomizableObjectInstance* Instance = Context->Instance.Get();
-	const UCustomizableObject* Object = Context->Object.Get();
 
-	if (!IsValid(Instance))
+	UCustomizableObjectInstance* Instance = Context->Instance.Get();
+	if (!CanUpdate(Instance))
 	{
 		Context->UpdateResult = EUpdateResult::Error;
 		FinishUpdateGlobal(Context);
@@ -3508,14 +3544,26 @@ bool UCustomizableObjectSystemPrivate::TryStartUpdate(const TSharedRef<FUpdateCo
 	// Skip update, the requested update is equal to the running update.
 	if (Context->GetCapturedDescriptorHash() == InstancePrivate->CommittedDescriptorHash)
 	{
-		Context->UpdateResult = EUpdateResult::Success;
-		UpdateSkeletalMesh(Context);
-		FinishUpdateGlobal(Context);
-		return false;
-	}
+		if (CVarEnableUpdateOptimization.GetValueOnGameThread()) // TODO Remove hotfix: UE-218957 
+		{
+			Context->bOptimizedUpdate = true;
 
-	Context->Parameters = Context->GetCapturedDescriptor().GetParameters();
+			// The user may have changed the AttachParent and we need to re-customize it.
+			// In case nothing need to be re-customized, the update will be considered ErrorOptimized.
+			UpdateSkeletalMesh(Context);
+			Context->UpdateResult = Context->AttachedParentUpdated.IsEmpty() ? EUpdateResult::ErrorOptimized : EUpdateResult::Success;
 
+			FinishUpdateGlobal(Context);
+			return true;
+		}
+		else
+		{
+			UpdateSkeletalMesh(Context);
+			FinishUpdateGlobal(Context);
+			return true;
+		}
+	}
+	
 	StartUpdate(Context);
 
 	return true;
@@ -3526,41 +3574,20 @@ bool UCustomizableObjectSystemPrivate::TryStartAutomaticUpdateForInstance(UCusto
 {
 	MUTABLE_CPUPROFILER_SCOPE(TryStartAutomaticUpdateForInstance);
 
+	TryAutomaticCompilation(&Instance);
+	
 	check(!CurrentMutableOperation); // Can not start an update if there is already another in progress
 
-	UCustomizableObject* Object = Instance.GetCustomizableObject();
-	if (!IsValid(Object))
-	{
-		return false;
-	}
-
-#if WITH_EDITOR
-	if (!Object->IsCompiled() &&
-		Object->GetPrivate()->CompilationResult != ECompilationResultPrivate::Errors) // Avoid constantly retry failed compilations.
-	{
-		if (ICustomizableObjectEditorModule* EditorModule = ICustomizableObjectEditorModule::Get())
-		{
-			FCompileParams CompileParams;
-			CompileParams.bSkipIfCompiled = true;
-
-			EditorModule->CompileCustomizableObject(*Object, &CompileParams, true, false);
-		}
-	}
-#endif
-
-	UCustomizableInstancePrivate* InstancePrivate = Instance.GetPrivate();
-	if (!InstancePrivate->CanUpdateInstance())
+	if (!CanUpdate(&Instance))
 	{
 		return false;
 	}
-
-	const TSharedRef<FUpdateContextPrivate> Context = MakeShared<FUpdateContextPrivate>(Instance, InstancePrivate->GetDescriptor());
-	if (!Context->IsContextValid())
-	{
-		return false;
-	}
-
-	if (Context->GetCapturedDescriptorHash() == InstancePrivate->CommittedDescriptorHash)
+	
+	const TSharedRef<FUpdateContextPrivate> Context = MakeShared<FUpdateContextPrivate>();
+	Context->Instance = &Instance;
+	Context->Init(true);
+	
+	if (Context->GetCapturedDescriptorHash() == Instance.GetPrivate()->CommittedDescriptorHash)
 	{
 		Context->bOptimizedUpdate = true;
 
@@ -3572,9 +3599,7 @@ bool UCustomizableObjectSystemPrivate::TryStartAutomaticUpdateForInstance(UCusto
 		FinishUpdateGlobal(Context);
 		return true;
 	}
-
-	Context->Parameters = Context->GetCapturedDescriptor().GetParameters();
-
+	
 	StartUpdate(Context);
 
 	return true;
diff --git a/Engine/Plugins/Mutable/Source/CustomizableObjectEditor/Private/MuCOE/CustomizableObjectInstanceBakingUtils.cpp b/Engine/Plugins/Mutable/Source/CustomizableObjectEditor/Private/MuCOE/CustomizableObjectInstanceBakingUtils.cpp
--- a/Engine/Plugins/Mutable/Source/CustomizableObjectEditor/Private/MuCOE/CustomizableObjectInstanceBakingUtils.cpp
+++ b/Engine/Plugins/Mutable/Source/CustomizableObjectEditor/Private/MuCOE/CustomizableObjectInstanceBakingUtils.cpp
@@ -420,12 +420,12 @@ void ScheduleInstanceUpdateForBaking(UCustomizableObjectInstance& InInstance, FI
 	PrepareUnrealCompression();
 	
 	// Schedule the update
-	const TSharedRef<FUpdateContextPrivate> Context = MakeShared<FUpdateContextPrivate>(InInstance);
+	const TSharedRef<FUpdateContextPrivate> Context = MakeShared<FUpdateContextPrivate>();
+	Context->Instance = &InInstance;
 	Context->bForce = true;
 	Context->bBake = true;
 	Context->bIsHighPriority = true;
 	Context->UpdateNativeCallback = InInstanceUpdateDelegate;
-	Context->bProgressiveMipStreamingEnabled = false;
 	Context->PixelFormatOverride = UnrealPixelFormatFunc;
 
 	UCustomizableObjectSystemPrivate* SystemPrivate = UCustomizableObjectSystem::GetInstance()->GetPrivate();
