[ContentBrowser] New Add Menu Mutable Menu

#Synced-CL 52417896

diff --git a/Engine/Plugins/Mutable/Source/CustomizableObjectEditor/Private/MuCOE/AssetDefinition_CustomizableObject.cpp b/Engine/Plugins/Mutable/Source/CustomizableObjectEditor/Private/MuCOE/AssetDefinition_CustomizableObject.cpp
--- a/Engine/Plugins/Mutable/Source/CustomizableObjectEditor/Private/MuCOE/AssetDefinition_CustomizableObject.cpp
+++ b/Engine/Plugins/Mutable/Source/CustomizableObjectEditor/Private/MuCOE/AssetDefinition_CustomizableObject.cpp
@@ -40,10 +40,9 @@ EAssetCommandResult UAssetDefinition_CustomizableObject::OpenAssets(const FAsset
 
 TConstArrayView<FAssetCategoryPath> UAssetDefinition_CustomizableObject::GetAssetCategories() const
 {
-	static const std::initializer_list<FAssetCategoryPath> Categories =
+	static const auto Categories =
 	{
-		// Asset can be found inside the Mutable submenu 
-		NSLOCTEXT("AssetTypeActions", "Mutable", "Mutable")
+		FAssetCategoryPath(NSLOCTEXT("AssetDefinition", "CustomizableObjectCategory", "Mutable"))
 	};
 	return Categories;
 }
