[mutable] Optimized chains of conditionals into switches.
- Fixed crash caused by the usage of a pointer to a potentially null FInstance object.
- Added a couple of optimizations that take care of converting chains of conditional operations that point to the same condition to switch operations.
- Removed no longer used includes.

#rnx
#Synced-CL 52336947

diff --git a/Engine/Plugins/Mutable/Source/CustomizableObjectEditor/Private/MuCOE/SMutableCodeViewer.cpp b/Engine/Plugins/Mutable/Source/CustomizableObjectEditor/Private/MuCOE/SMutableCodeViewer.cpp
--- a/Engine/Plugins/Mutable/Source/CustomizableObjectEditor/Private/MuCOE/SMutableCodeViewer.cpp
+++ b/Engine/Plugins/Mutable/Source/CustomizableObjectEditor/Private/MuCOE/SMutableCodeViewer.cpp
@@ -5,7 +5,6 @@
 #include "SMutableMaterialViewer.h"
 #include "Framework/MultiBox/MultiBoxBuilder.h"
 #include "Framework/Views/TableViewMetadata.h"
-#include "Misc/Paths.h"
 #include "MuCO/UnrealConversionUtils.h"
 #include "MuCO/UnrealToMutableTextureConversionUtils.h"
 #include "MuCOE/SMutableBoolViewer.h"
@@ -36,7 +35,6 @@
 #include "Widgets/Layout/SScrollBox.h"
 #include "Widgets/Views/STreeView.h"
 #include "Widgets/Input/SSearchBox.h"
-#include "Engine/SkeletalMesh.h"
 #include "Internationalization/Regex.h"
 #include "MuCO/CustomizableObjectInstancePrivate.h"
 #include "MuR/PassthroughObject.h"
@@ -3342,6 +3340,14 @@ void FMutableCodeTreeElement::GenerateLabelText()
 		OpName += FString::Printf(TEXT(" mips %d"), Args.LODs);
 		break;
 	}
+		
+	case UE::Mutable::Private::EOpType::BO_EQUAL_INT_CONST:
+	{
+		UE::Mutable::Private::FOperation::BoolEqualScalarConstArgs Args = Program.GetOpArgs<UE::Mutable::Private::FOperation::BoolEqualScalarConstArgs>(MutableOperation);
+		OpName += TEXT(" const: ");
+		OpName += FString::FromInt(Args.Constant);
+		break;
+	}
 
 	default:
 		break;
diff --git a/Engine/Plugins/Mutable/Source/MutableRuntime/Private/MuR/CodeRunner.cpp b/Engine/Plugins/Mutable/Source/MutableRuntime/Private/MuR/CodeRunner.cpp
--- a/Engine/Plugins/Mutable/Source/MutableRuntime/Private/MuR/CodeRunner.cpp
+++ b/Engine/Plugins/Mutable/Source/MutableRuntime/Private/MuR/CodeRunner.cpp
@@ -740,7 +740,7 @@ namespace UE::Mutable::Private
 							TManagedPtr<const FInstance> pLOD = LoadInstance(FCacheAddress(LODAddress, Item));
 
                             // In a degenerated case, the returned pLOD may not have an LOD inside
- 							if (!pLOD->Components.IsEmpty() && !pLOD->Components[0].LODs.IsEmpty())
+ 							if (pLOD && !pLOD->Components.IsEmpty() && !pLOD->Components[0].LODs.IsEmpty())
 							{
 								Result->Components[ComponentIndex].LODs[InstanceLODIndex] = pLOD->Components[0].LODs[0];
 							}
diff --git a/Engine/Plugins/Mutable/Source/MutableTools/Private/MuT/CodeOptimiser_Logic.cpp b/Engine/Plugins/Mutable/Source/MutableTools/Private/MuT/CodeOptimiser_Logic.cpp
--- a/Engine/Plugins/Mutable/Source/MutableTools/Private/MuT/CodeOptimiser_Logic.cpp
+++ b/Engine/Plugins/Mutable/Source/MutableTools/Private/MuT/CodeOptimiser_Logic.cpp
@@ -1,5 +1,6 @@
 // Copyright Epic Games, Inc. All Rights Reserved.
 
+#include "ASTOpBoolAnd.h"
 #include "Misc/AssertionMacros.h"
 #include "MuR/MutableTrace.h"
 #include "MuR/Operations.h"
@@ -7,16 +8,263 @@
 #include "MuR/RefCounted.h"
 #include "MuT/AST.h"
 #include "MuT/ASTOpConditional.h"
-#include "MuT/ASTOpInstanceAdd.h"
 #include "MuT/ASTOpSwitch.h"
-#include "MuT/ASTOpMeshMerge.h"
 #include "MuT/ASTOpBoolEqualIntConst.h"
 #include "MuT/CodeOptimiser.h"
 
 
 namespace UE::Mutable::Private
 {
+	/**
+	 * Enumeration representing the branches that an ASTOpBoolAnd operation has.
+	 */
+	enum class EASTOpBoolAndBranch : int32
+	{
+		A = 0,
+		B,
+	};
+	
+	
+	Ptr<ASTOpSwitch> OptimizeConditionalChainBranch(Ptr<ASTOp>& n, const EASTOpBoolAndBranch TargetBoolAndOpBranch)
+	{
+		ASTOpConditional* RootConditional = nullptr;
+		if (n && n->IsConditional())
+		{
+			RootConditional = static_cast<ASTOpConditional*>(n.get());
+		}
+
+		Ptr<ASTOpSwitch> SwitchOP = nullptr;
+		
+		// Crete the switch object first and set up the type and the variable to later be used to evaluate what case is the active one
+		if ( RootConditional && RootConditional->condition &&  
+			 (RootConditional->condition->GetOpType()==EOpType::BO_EQUAL_INT_CONST || RootConditional->condition->GetOpType() == EOpType::BO_AND)
+			 &&
+			 RootConditional->no
+			 &&
+			 RootConditional->no->GetOpType()==RootConditional->GetOpType()
+			)
+		{
+			SwitchOP = new ASTOpSwitch();
+			SwitchOP->Type = GetSwitchForType(GetOpDataType(RootConditional->GetOpType()));
+			
+			// Look for the conditional node of the topConditional and use the value set on it as the variable for the switch. 
+			{
+				// Simple Case, BO_EQUAL_INT_CONST
+				if (RootConditional->condition->GetOpType()==EOpType::BO_EQUAL_INT_CONST )
+				{
+					ASTOpBoolEqualIntConst* TopConditionalCondition = static_cast<ASTOpBoolEqualIntConst*>(RootConditional->condition.child().get());
+					SwitchOP->Variable = TopConditionalCondition->Value.child();
+				}
+				
+				// The condition is an AND operation
+				else if (RootConditional->condition->GetOpType() == EOpType::BO_AND)
+				{
+					ASTOpBoolAnd* firstCompare = static_cast<ASTOpBoolAnd*>(RootConditional->condition.child().get());
+				
+					// Locate the Equal_cost of the And  use it as the variable of the switch
+					if (TargetBoolAndOpBranch == EASTOpBoolAndBranch::A && 
+						firstCompare->A.child()->GetOpType() == EOpType::BO_EQUAL_INT_CONST)
+					{
+						ASTOpBoolEqualIntConst* EqualIntConst = static_cast<ASTOpBoolEqualIntConst*>(firstCompare->A.child().get());
+						check(EqualIntConst);
+						
+						// Now the variable of the switch will use the "parameter" set in the value branch of the const
+						SwitchOP->Variable = EqualIntConst->Value.child();
+					}
+					else if (TargetBoolAndOpBranch == EASTOpBoolAndBranch::B && 
+						firstCompare->B.Child->GetOpType()==EOpType::BO_EQUAL_INT_CONST)
+					{
+						ASTOpBoolEqualIntConst* EqualIntConst = static_cast<ASTOpBoolEqualIntConst*>(firstCompare->B.child().get());
+						check(EqualIntConst);
+						
+						SwitchOP->Variable = EqualIntConst->Value.child();
+					}
+				}
+			}
+			
+			// node not compatible with this optimization as no variable for the switch could be found
+			if (!SwitchOP->Variable.child().get())
+			{
+				return nullptr;
+			}
+
+			TArray<Ptr<ASTOpConditional>> NewConditionals;
+			
+	        // Iterate over the child objects and add the ones that use the same variable as the switch as cases for it
+			Ptr<ASTOp> current = n;
+	        while(current)
+			{
+				bool bValid = false;
+
+				ASTOpConditional* CurrentConditional = nullptr;
+				if (current && current->IsConditional())
+				{
+					CurrentConditional = static_cast<ASTOpConditional*>(current.get());
+				}
+	        	
+				// Simple case: the conditional is an BO_EQUAL_INT_CONST
+				if ( CurrentConditional
+					 && 
+					 CurrentConditional->GetOpType()==RootConditional->GetOpType()
+					 &&
+					 CurrentConditional->condition
+					 &&
+					 CurrentConditional->condition->GetOpType()==EOpType::BO_EQUAL_INT_CONST )
+				{
+					ASTOpBoolEqualIntConst* Compare = static_cast<ASTOpBoolEqualIntConst*>(CurrentConditional->condition.child().get());
+					check(Compare);
+
+					if ( Compare->Value.child() == SwitchOP->Variable.child() )
+					{
+						const auto CaseSearchLambda = [NewCaseIndex = Compare->Constant](const ASTOpSwitch::FCase& Other)->bool
+						{
+							return Other.Condition == NewCaseIndex;
+						};           		
+					
+						if (!SwitchOP->Cases.ContainsByPredicate(CaseSearchLambda))
+						{
+							SwitchOP->Cases.Emplace( Compare->Constant, SwitchOP, CurrentConditional->yes.child() );
+
+							current = CurrentConditional->no.child();
+							bValid = true;
+						}
+						else
+						{
+							// If we added this into the cases it would collide with another case's condition. Abort optimization.
+							return nullptr;
+						}
+					}
+				}
+				// A bit more complex: 
+				else if ( CurrentConditional
+				 && 
+				 CurrentConditional->GetOpType()==RootConditional->GetOpType()
+				 &&
+				 CurrentConditional->condition
+				 &&
+				 CurrentConditional->condition->GetOpType()==EOpType::BO_AND )
+				{
+					ASTOpBoolAnd* AndOp = static_cast<ASTOpBoolAnd*>(CurrentConditional->condition.child().get());
+					check(AndOp);
+					
+					int32 SwitchConditionValue = -1;
+					Ptr<ASTOp> CaseChildCondition = nullptr;
+					bool bSwitchConditionValueFound = false;
+					
+					if (TargetBoolAndOpBranch == EASTOpBoolAndBranch::A && AndOp->A.child()->GetOpType() == EOpType::BO_EQUAL_INT_CONST)
+					{
+						ASTOpBoolEqualIntConst* Compare = static_cast<ASTOpBoolEqualIntConst*>(AndOp->A.child().get());
+						check(Compare);
+						
+						// Ensure the value of this const also points to the same variable of the switch. If not just ignore it
+						if (Compare->Value.child() == SwitchOP->Variable.child())
+						{
+							const int32 ACaseConditionValue = Compare->Constant;
+							
+							const auto CaseSearchLambdaA = [ACaseConditionValue](const ASTOpSwitch::FCase& Other)->bool
+							{
+								return Other.Condition == ACaseConditionValue;
+							};
+							
+							if (!SwitchOP->Cases.ContainsByPredicate(CaseSearchLambdaA))
+							{
+								SwitchConditionValue = ACaseConditionValue;
+								bSwitchConditionValueFound = true;
+							}
+							else
+							{
+								// If we added this into the cases it would collide with another case's condition. Abort optimization.
+								return nullptr;
+							}
+							
+							// since A is the equal then B is the condition we want now t use with this node
+							CaseChildCondition = AndOp->B.child();
+						}
+					}
+					
+					// A is not a valid case, try B branch
+					else if (TargetBoolAndOpBranch == EASTOpBoolAndBranch::B && AndOp->B.Child->GetOpType()==EOpType::BO_EQUAL_INT_CONST)
+					{
+						ASTOpBoolEqualIntConst* Compare = static_cast<ASTOpBoolEqualIntConst*>(AndOp->B.child().get());
+						check(Compare);
+						
+						if (Compare->Value.child() == SwitchOP->Variable.child())
+						{	
+							const int32 BCaseConditionValue = Compare->Constant;
+															
+							const auto CaseSearchLambdaB = [BCaseConditionValue](const ASTOpSwitch::FCase& Other)->bool
+							{
+								return Other.Condition == BCaseConditionValue;
+							};
+								
+							if (!SwitchOP->Cases.ContainsByPredicate(CaseSearchLambdaB))
+							{
+								SwitchConditionValue = BCaseConditionValue;
+								bSwitchConditionValueFound = true;
+							}
+							else
+							{
+								// If we added this into the cases it would collide with another case's condition. Abort optimization.
+								return nullptr;
+							}
+							
+							// since B is the equal then A is the condition we want now t use with this node
+							CaseChildCondition = AndOp->A.child();
+						}
+					}
+					
+					// Ony add a new case with condition index if the index has not already been added
+					if (bSwitchConditionValueFound)
+					{
+						// Make a new conditional to use as the new case where only the Yes section is there with the condition that did 
+						// coexist with the EQ_Const of the BO_And. Since the switch already hadles the EQ_Const then the only condition we need is the other one
+						Ptr<ASTOpConditional> NewConditional = new ASTOpConditional();
+						NewConditional->type = CurrentConditional->type;
+						NewConditional->yes = ASTChild( NewConditional, CurrentConditional->yes.child().get());
+						NewConditional->no = ASTChild( NewConditional, nullptr);
+						NewConditional->condition = CaseChildCondition;
+						
+						NewConditionals.Add( NewConditional);
+						
+						// Add the NewConditional node as a case for the switch.
+						// This case has:
+						//	- condition : the constant value of an ASTOpBoolEqualIntConst extracted from the Branch A or Branch B of the BO_AND op
+						//	- parent : the switch node it will be part of
+						//	- branch : a branch that has as root a new conditional OP NewConditional that has as condition the complementary branch from the one 
+						//				used to get the "condition" integer value.
+						//				it will look something like this :
+						//				IF_CONDITIONAL == "NewConditional"
+						//					->	Condition : A if case condition is B, B otherwise.
+						//					->	YES : The "yes" from the conditional node we are processing now. 
+						//					->	NO : Default case of the switch. (added later when iterating the NewConditionals array)NewConditional
+						
+						SwitchOP->Cases.Emplace( SwitchConditionValue, SwitchOP, NewConditional);
+						
+						// make the next branch to be processed the other side of the current conditional as we have added to the switch a path including the YES side
+						current = CurrentConditional->no.child();			
+						bValid = true;
+					}
+				}
 
+				if (!bValid)
+				{
+					SwitchOP->Default = current;
+
+					// Make the no branch have the current node branch once the switch setup is done
+					for (Ptr<ASTOpConditional>& NewConditional : NewConditionals)
+					{
+						NewConditional->no = SwitchOP->Default.child();
+					}
+					
+					current = nullptr;
+				}
+			}
+		}
+		
+		return SwitchOP;
+	}
+	
+	
 	bool LocalLogicOptimiserAST(ASTOpList& roots)
     {
         MUTABLE_CPUPROFILER_SCOPE(LocalLogicOptimiserAST);
@@ -91,88 +339,106 @@ namespace UE::Mutable::Private
         //-----------------------------------------------------------------------------------------
         {
             MUTABLE_CPUPROFILER_SCOPE(ConditionalToSwitch);
-            ASTOp::Traverse_TopDown_Unique_Imprecise( roots, [&](Ptr<ASTOp>& n)
+        	
+        	constexpr int32 MIN_CONDITIONS_TO_CREATE_SWITCH = 3;
+        	
+        	// Force optimization of the A branch of the conditional BO_AND operation
+            ASTOp::Traverse_TopDown_Unique_Imprecise( roots, [&](Ptr<ASTOp>& CurrentOp)
             {
-                bool bRecurse = true;
-
-				ASTOpConditional* topConditional = nullptr;
-				if (n && n->IsConditional())
+				// Try processing first the A branch of the BO_AND (if found)
+				const Ptr<ASTOpSwitch> SwitchOp = OptimizeConditionalChainBranch(CurrentOp, EASTOpBoolAndBranch::A);
+								
+				if (SwitchOp && SwitchOp->Cases.Num() >= MIN_CONDITIONS_TO_CREATE_SWITCH)
 				{
-					topConditional = static_cast<ASTOpConditional*>(n.get());
+					ASTOp::Replace(CurrentOp,SwitchOp);
+					CurrentOp = SwitchOp;
+					bModified = true;
 				}
-
-                if (topConditional && topConditional->condition)
-                {
-                    if ( topConditional->condition->GetOpType()==EOpType::BO_EQUAL_INT_CONST
-                         &&
-                         topConditional->no
-                         &&
-                         topConditional->no->GetOpType()==topConditional->GetOpType()
-                         )
-                    {
-                        Ptr<ASTOpSwitch> switchOp = new ASTOpSwitch();
-                        switchOp->Type = GetSwitchForType(GetOpDataType(topConditional->GetOpType()));
-
-						ASTOpBoolEqualIntConst* firstCompare = static_cast<ASTOpBoolEqualIntConst*>(topConditional->condition.child().get());
-                        switchOp->Variable = firstCompare->Value.child();
-
-						Ptr<ASTOp> current = n;
-                        while(current)
-                        {
-                            bool bValid = false;
-
-							ASTOpConditional* conditional = nullptr;
-							if (current && current->IsConditional())
-							{
-								conditional = static_cast<ASTOpConditional*>(current.get());
-							}
-
-                            if ( conditional
-                                 &&
-                                 conditional->GetOpType()==topConditional->GetOpType()
-                                 &&
-                                 conditional->condition
-                                 &&
-                                 conditional->condition->GetOpType()==EOpType::BO_EQUAL_INT_CONST )
-                            {
-								ASTOpBoolEqualIntConst* compare = static_cast<ASTOpBoolEqualIntConst*>(conditional->condition.child().get());
-                                check(compare);
-                                if (compare)
-                                {
-                                    auto compareValue = compare->Value.child();
-                                    check(compare);
-                                    if ( compareValue == switchOp->Variable.child() )
-                                    {
-                                        switchOp->Cases.Emplace( compare->Constant, switchOp, conditional->yes.child() );
-
-                                        current = conditional->no.child();
-										bValid = true;
-                                    }
-                                }
-                            }
-
-                            if (!bValid)
-                            {
-                                switchOp->Default = current;
-                                current = nullptr;
-                            }
-                        }
-
-                        const int MIN_CONDITIONS_TO_CREATE_SWITCH = 3;
-                        if (switchOp->Cases.Num()>=MIN_CONDITIONS_TO_CREATE_SWITCH)
-                        {
-                            ASTOp::Replace(n,switchOp);
-                            n = switchOp;
-                            bModified = true;
-                        }
-                    }
-
-                }
-
+            		
+            	constexpr bool bRecurse = true;
                 return bRecurse;
             });
+        	
+        	// Force optimization of the B branch of the conditional BO_AND operation
+        	ASTOp::Traverse_TopDown_Unique_Imprecise( roots, [&](Ptr<ASTOp>& CurrentOp)
+			{			
+				// Then do the same for the B branch of the BO_AND (also if found)
+				const Ptr<ASTOpSwitch> SwitchOp = OptimizeConditionalChainBranch(CurrentOp, EASTOpBoolAndBranch::B);
+			
+				if (SwitchOp && SwitchOp->Cases.Num() >= MIN_CONDITIONS_TO_CREATE_SWITCH)
+				{
+					ASTOp::Replace(CurrentOp,SwitchOp);
+					CurrentOp = SwitchOp;
+					bModified = true;
+				}
+            		
+				constexpr bool bRecurse = true;
+				return bRecurse;
+			});
         }
 
+		
+		
+		ASTOp::Traverse_BottomUp_Unique( roots, [&](Ptr<ASTOp>& CurrentOp)
+		{	
+			if (!CurrentOp->IsSwitch())
+			{
+				return;
+		  	}
+				
+			ASTOpSwitch* Switch = static_cast<ASTOpSwitch*>(CurrentOp.get());
+
+			Ptr<ASTOpSwitch> NewSwitch = Clone<ASTOpSwitch>(Switch);
+			NewSwitch->Cases.Reset();
+				
+			ASTOp* FirstCaseCondition = nullptr;
+				
+			for (int32 Index = 0; Index < Switch->Cases.Num(); ++Index)
+			{
+				const ASTOpSwitch::FCase& Case = Switch->Cases[Index];
+				
+				if (!Case.Branch || !Case.Branch->IsConditional())
+				{
+					return;
+				}
+				
+				ASTOpConditional* CaseConditional = static_cast<ASTOpConditional*>(Case.Branch.child().get());
+				
+				if (!FirstCaseCondition)
+				{
+					FirstCaseCondition = CaseConditional->condition.child().get();
+				}
+				
+				if (FirstCaseCondition != CaseConditional->condition.child())
+				{
+					return;
+				}
+				
+				if (CaseConditional->no != Switch->Default)
+				{
+					return;
+				}
+				
+				if (!CaseConditional->yes)
+				{
+					return;
+				}
+				
+				NewSwitch->Cases.Emplace(Case.Condition, NewSwitch, CaseConditional->yes.child());
+			}
+				
+			// All conditionals should have a condition.
+			check(FirstCaseCondition);
+				
+			Ptr<ASTOpConditional> NewConditional = new ASTOpConditional();
+			NewConditional->type = GetConditionalForType(GetOpDataType(NewSwitch->GetOpType()));
+			NewConditional->condition = FirstCaseCondition;
+			NewConditional->yes = NewSwitch;
+			NewConditional->no = NewSwitch->Default.child();
+				
+			ASTOp::Replace(CurrentOp, NewConditional);
+		});
+		
 
         // Float operations up switches, to tidy up the code and reduce its size.
 		// TODO?
diff --git a/Engine/Plugins/Mutable/Source/MutableTools/Private/MuT/CodeOptimiser_Parameters.cpp b/Engine/Plugins/Mutable/Source/MutableTools/Private/MuT/CodeOptimiser_Parameters.cpp
--- a/Engine/Plugins/Mutable/Source/MutableTools/Private/MuT/CodeOptimiser_Parameters.cpp
+++ b/Engine/Plugins/Mutable/Source/MutableTools/Private/MuT/CodeOptimiser_Parameters.cpp
@@ -13,9 +13,7 @@
 #include "MuT/ASTOpImagePlainColor.h"
 #include "MuT/ASTOpImageDisplace.h"
 #include "MuT/ASTOpMaterialModify.h"
-#include "MuT/ASTOpMeshMorph.h"
 #include "MuT/ASTOpInstanceAdd.h"
-#include "MuT/ASTOpLayoutFromMesh.h"
 #include "MuT/ASTOpParameter.h"
 #include "MuT/ASTOpConstantColor.h"
 #include "MuT/CodeOptimiser.h"
@@ -31,11 +29,9 @@
 #include "MuR/MutableRuntimeModule.h"
 
 #include "HAL/PlatformCrt.h"
-#include "HAL/PlatformMath.h"
 #include "Logging/LogCategory.h"
 #include "Logging/LogMacros.h"
 #include "Misc/AssertionMacros.h"
-#include "Trace/Detail/Channel.h"
 
 #include <array>
 
@@ -373,13 +369,13 @@ namespace UE::Mutable::Private
             // conditional down
             bool optimised = false;
 
-            if ( !HasRuntimeParamVisitor.HasAny( TypedOp->condition.child() ) )
+            if ( !HasRuntimeParamVisitor.HasAny( TypedOp->condition.child() ) ) 
             {
-                EOpType yesType = TypedOp->yes->GetOpType();
-                EOpType noType = TypedOp->no->GetOpType();
+                EOpType yesType = TypedOp->yes.child().get() ? TypedOp->yes->GetOpType() : EOpType::NONE;
+                EOpType noType = TypedOp->no.child().get() ? TypedOp->no->GetOpType() : EOpType::NONE;
 
-                bool yesHasAny = HasRuntimeParamVisitor.HasAny( TypedOp->yes.child() );
-                bool noHasAny = HasRuntimeParamVisitor.HasAny( TypedOp->no.child() );
+                bool yesHasAny = yesType != EOpType::NONE && HasRuntimeParamVisitor.HasAny( TypedOp->yes.child() );
+                bool noHasAny = noType != EOpType::NONE && HasRuntimeParamVisitor.HasAny( TypedOp->no.child() );
 
                 if ( !optimised && yesHasAny && noHasAny && yesType==noType)
                 {
