[mutable] Fix potential crash when accessing null object and to improve code consistency.

#rnx
#Synced-CL 52336947

diff --git a/Engine/Plugins/Mutable/Source/MutableTools/Private/MuT/CodeOptimiser_Logic.cpp b/Engine/Plugins/Mutable/Source/MutableTools/Private/MuT/CodeOptimiser_Logic.cpp
--- a/Engine/Plugins/Mutable/Source/MutableTools/Private/MuT/CodeOptimiser_Logic.cpp
+++ b/Engine/Plugins/Mutable/Source/MutableTools/Private/MuT/CodeOptimiser_Logic.cpp
@@ -183,7 +183,7 @@ namespace UE::Mutable::Private
 					}
 					
 					// A is not a valid case, try B branch
-					else if (TargetBoolAndOpBranch == EASTOpBoolAndBranch::B && AndOp->B.Child->GetOpType()==EOpType::BO_EQUAL_INT_CONST)
+					else if (TargetBoolAndOpBranch == EASTOpBoolAndBranch::B && AndOp->B.child()->GetOpType()==EOpType::BO_EQUAL_INT_CONST)
 					{
 						ASTOpBoolEqualIntConst* Compare = static_cast<ASTOpBoolEqualIntConst*>(AndOp->B.child().get());
 						check(Compare);
@@ -490,4 +490,3 @@ namespace UE::Mutable::Private
     }
 
 }
-
