[mutable] Commented out optimization code that seems to be the cause (direct or indirect) of a  crash while the RenderSections of an SKM are being set up by mutable.

#rnx
#Synced-CL 52498324

diff --git a/Engine/Plugins/Mutable/Source/MutableTools/Private/MuT/CodeOptimiser_Logic.cpp b/Engine/Plugins/Mutable/Source/MutableTools/Private/MuT/CodeOptimiser_Logic.cpp
--- a/Engine/Plugins/Mutable/Source/MutableTools/Private/MuT/CodeOptimiser_Logic.cpp
+++ b/Engine/Plugins/Mutable/Source/MutableTools/Private/MuT/CodeOptimiser_Logic.cpp
@@ -378,66 +378,66 @@ namespace UE::Mutable::Private
         }
 
 		
-		
-		ASTOp::Traverse_BottomUp_Unique( roots, [&](Ptr<ASTOp>& CurrentOp)
-		{	
-			if (!CurrentOp->IsSwitch())
-			{
-				return;
-		  	}
-				
-			ASTOpSwitch* Switch = static_cast<ASTOpSwitch*>(CurrentOp.get());
-
-			Ptr<ASTOpSwitch> NewSwitch = Clone<ASTOpSwitch>(Switch);
-			NewSwitch->Cases.Reset();
-				
-			ASTOp* FirstCaseCondition = nullptr;
-				
-			for (int32 Index = 0; Index < Switch->Cases.Num(); ++Index)
-			{
-				const ASTOpSwitch::FCase& Case = Switch->Cases[Index];
-				
-				if (!Case.Branch || !Case.Branch->IsConditional())
-				{
-					return;
-				}
-				
-				ASTOpConditional* CaseConditional = static_cast<ASTOpConditional*>(Case.Branch.child().get());
-				
-				if (!FirstCaseCondition)
-				{
-					FirstCaseCondition = CaseConditional->condition.child().get();
-				}
-				
-				if (FirstCaseCondition != CaseConditional->condition.child())
-				{
-					return;
-				}
-				
-				if (CaseConditional->no != Switch->Default)
-				{
-					return;
-				}
-				
-				if (!CaseConditional->yes)
-				{
-					return;
-				}
-				
-				NewSwitch->Cases.Emplace(Case.Condition, NewSwitch, CaseConditional->yes.child());
-			}
-				
-			// All conditionals should have a condition.
-			check(FirstCaseCondition);
-				
-			Ptr<ASTOpConditional> NewConditional = new ASTOpConditional();
-			NewConditional->type = GetConditionalForType(GetOpDataType(NewSwitch->GetOpType()));
-			NewConditional->condition = FirstCaseCondition;
-			NewConditional->yes = NewSwitch;
-			NewConditional->no = NewSwitch->Default.child();
-				
-			ASTOp::Replace(CurrentOp, NewConditional);
-		});
+		// TODO: This code needs to be revised as currently causes some instances to crash during the setup of the render sections (indirect consequence).
+		// ASTOp::Traverse_BottomUp_Unique( roots, [&](Ptr<ASTOp>& CurrentOp)
+		// {	
+		// 	if (!CurrentOp->IsSwitch())
+		// 	{
+		// 		return;
+		//   }
+		// 		
+		// 	ASTOpSwitch* Switch = static_cast<ASTOpSwitch*>(CurrentOp.get());
+		// 
+		// 	Ptr<ASTOpSwitch> NewSwitch = Clone<ASTOpSwitch>(Switch);
+		// 	NewSwitch->Cases.Reset();
+		// 		
+		// 	ASTOp* FirstCaseCondition = nullptr;
+		// 		
+		// 	for (int32 Index = 0; Index < Switch->Cases.Num(); ++Index)
+		// 	{
+		// 		const ASTOpSwitch::FCase& Case = Switch->Cases[Index];
+		// 		
+		// 		if (!Case.Branch || !Case.Branch->IsConditional())
+		// 		{
+		// 			return;
+		// 		}
+		// 		
+		// 		ASTOpConditional* CaseConditional = static_cast<ASTOpConditional*>(Case.Branch.child().get());
+		// 		
+		// 		if (!FirstCaseCondition)
+		// 		{
+		// 			FirstCaseCondition = CaseConditional->condition.child().get();
+		// 		}
+		// 		
+		// 		if (FirstCaseCondition != CaseConditional->condition.child())
+		// 		{
+		// 			return;
+		// 		}
+		// 		
+		// 		if (CaseConditional->no != Switch->Default)
+		// 		{
+		// 			return;
+		// 		}
+		// 		
+		// 		if (!CaseConditional->yes)
+		// 		{
+		// 			return;
+		// 		}
+		// 		
+		// 		NewSwitch->Cases.Emplace(Case.Condition, NewSwitch, CaseConditional->yes.child());
+		// 	}
+		// 		
+		// 	// All conditionals should have a condition.
+		// 	check(FirstCaseCondition);
+		// 		
+		// 	Ptr<ASTOpConditional> NewConditional = new ASTOpConditional();
+		// 	NewConditional->type = GetConditionalForType(GetOpDataType(NewSwitch->GetOpType()));
+		// 	NewConditional->condition = FirstCaseCondition;
+		// 	NewConditional->yes = NewSwitch;
+		// 	NewConditional->no = NewSwitch->Default.child();
+		// 		
+		// 	ASTOp::Replace(CurrentOp, NewConditional);
+		// });
 		
 
         // Float operations up switches, to tidy up the code and reduce its size.
