Fixed indeterministic asset registry queries during cook in Mutable

#Synced-CL 52502928

diff --git a/Engine/Plugins/Mutable/Source/CustomizableObjectEditor/Private/MuCOE/CustomizableObjectCompiler.cpp b/Engine/Plugins/Mutable/Source/CustomizableObjectEditor/Private/MuCOE/CustomizableObjectCompiler.cpp
--- a/Engine/Plugins/Mutable/Source/CustomizableObjectEditor/Private/MuCOE/CustomizableObjectCompiler.cpp
+++ b/Engine/Plugins/Mutable/Source/CustomizableObjectEditor/Private/MuCOE/CustomizableObjectCompiler.cpp
@@ -558,6 +558,8 @@ void ProcessChildObjectsRecursively(const UCustomizableObject* ParentObject, FMu
 	FARFilter Filter;
 	Filter.PackageNames = MoveTemp(ReferencedObjectNames);
 	Filter.ClassPaths = { UCustomizableObject::StaticClass()->GetClassPathName() };
+	Filter.bIncludeOnlyOnDiskAssets = IsRunningCookCommandlet();	// When cooking, only enumerate on-disk assets- avoids indeterministic results
+
 	AssetRegistryModule.Get().GetAssets(Filter, AssetDataArray);
 
 	// Required to be deterministic.
diff --git a/Engine/Plugins/Mutable/Source/CustomizableObjectEditor/Private/MuCOE/CustomizableObjectEditorModule.cpp b/Engine/Plugins/Mutable/Source/CustomizableObjectEditor/Private/MuCOE/CustomizableObjectEditorModule.cpp
--- a/Engine/Plugins/Mutable/Source/CustomizableObjectEditor/Private/MuCOE/CustomizableObjectEditorModule.cpp
+++ b/Engine/Plugins/Mutable/Source/CustomizableObjectEditor/Private/MuCOE/CustomizableObjectEditorModule.cpp
@@ -912,6 +912,7 @@ TMap<FName, FGuid> FCustomizableObjectEditorModule::GetParticipatingObjects(cons
 
 		FARFilter Filter;
 		Filter.bRecursivePaths = bRecursively;
+		Filter.bIncludeOnlyOnDiskAssets = IsRunningCookCommandlet();	// When cooking, only enumerate on-disk assets- avoids indeterministic results
 		
 		for (const FName& PackageName : PackageDependencies)
 		{
