PR #14573: [Mutable] Add Serialization Support for Material Params in LoadDescriptor

#Synced-CL 52538286
#Synced-CL 52538286

#Synced-CL 52538286

#Synced-CL 52538286

#Synced-CL 52538286

diff --git a/Engine/Plugins/Mutable/Source/CustomizableObject/Private/MuCO/CustomizableObjectInstanceDescriptor.cpp b/Engine/Plugins/Mutable/Source/CustomizableObject/Private/MuCO/CustomizableObjectInstanceDescriptor.cpp
--- a/Engine/Plugins/Mutable/Source/CustomizableObject/Private/MuCO/CustomizableObjectInstanceDescriptor.cpp
+++ b/Engine/Plugins/Mutable/Source/CustomizableObject/Private/MuCO/CustomizableObjectInstanceDescriptor.cpp
@@ -563,6 +563,31 @@ void FCustomizableObjectInstanceDescriptor::LoadDescriptor(FArchive& Ar)
 			break;
 		}
 			
+		case EMutableParameterType::Material:
+		{
+			FString Value;
+			TArray<FString> RangeValues;
+			Ar << Value;
+			Ar << RangeValues;
+			for (FCustomizableObjectMaterialParameterValue& P : MaterialParameters)
+			{
+				if (P.ParameterName == Name)
+				{
+					P.ParameterValue = ToObject<UMaterialInterface>(Value);
+					
+					P.ParameterRangeValues.Empty(RangeValues.Num());
+					Algo::Transform(RangeValues, P.ParameterRangeValues, [](const FString& Element)
+					{
+						return ToObject<UMaterialInterface>(Element);
+					});
+		
+					break;
+				}
+			}
+				
+			break;
+		}
+			
 		case EMutableParameterType::None:
 		{
 			UE_LOGF(LogMutable, Warning, "Found unknown parameter [%ls] when loading a Customizable Object Instance Descriptor", *Name);
