[Mutable] Fix crash when accessing CO outside Game Thread.

#rnx
#Synced-CL 52538286

diff --git a/Engine/Plugins/Mutable/Source/CustomizableObject/Private/MuCO/CustomizableObjectSystem.cpp b/Engine/Plugins/Mutable/Source/CustomizableObject/Private/MuCO/CustomizableObjectSystem.cpp
--- a/Engine/Plugins/Mutable/Source/CustomizableObject/Private/MuCO/CustomizableObjectSystem.cpp
+++ b/Engine/Plugins/Mutable/Source/CustomizableObject/Private/MuCO/CustomizableObjectSystem.cpp
@@ -1986,8 +1986,6 @@ namespace impl
 
 		// Setup Component and mesh descriptors
 		{
-			OperationData->InitComponentDescriptors(OperationData->Object->GetComponentCount());
-            		
 			const TMap<FName, uint8>& NumLODsAvailable = OperationData->NumLODsAvailable;
 	
 			for (int32 ComponentIndex = 0; ComponentIndex < OperationData->MutableInstance->GetComponentCount(); ++ComponentIndex)
@@ -2218,6 +2216,14 @@ namespace impl
 			return;
 		}
 		
+		// TODO We should not access UObjects outside GT
+		const TStrongObjectPtr<UCustomizableObject> StrongObject = OperationData->Object.Pin();
+		if (!IsValid(StrongObject.Get()))
+		{
+			OperationData->bIsUpdateAborted = true;
+			return;
+		}
+		
 		const TSharedPtr<TArray<UE::Mutable::Private::FImageId>> RequestedGetImages = MakeShared<TArray<UE::Mutable::Private::FImageId>>();
 
 		UE::Tasks::FTask LastTask;
@@ -2244,7 +2250,7 @@ namespace impl
 				// Access the CO TextureCache and try to grab the texture from there if found.
 				if (OperationData->bUseCaches)
 				{
-					FTextureCache& TextureCache = OperationData->Object->GetPrivate()->TextureCache;
+					FTextureCache& TextureCache = StrongObject->GetPrivate()->TextureCache;
 					const FTextureCache::FId TextureCacheKey = FTextureCache::FId(Image->ImageID, OperationData->MipsToSkip, OperationData->bBake);
 
 					if (const TStrongObjectPtr<const UTexture> CachedTexture = TextureCache.Get(TextureCacheKey))
@@ -2841,6 +2847,7 @@ namespace impl
 		}
 #endif
 		
+		Operation->InitComponentDescriptors(Operation->Object->GetComponentCount());
 		
 		UE::Tasks::FTask Mutable_GetMeshTask = SystemPrivate->MutableTaskGraph.AddMutableThreadTask(
 		TEXT("Task_Mutable_GetMeshes"),
