[Mutable] Add MarkedToKeep flag to reuse expensive generated layouts as much as possible.

- Removed parent task options from layout scheduling so the same layouts generated from different runs go to the same cache address.
- Added cpu scopes for layout operations.
- Now Layouts generated from instruction LA_FORMMESH are marked to keep and are not cleared.

#rnx
#Synced-CL 52615216

diff --git a/Engine/Plugins/Mutable/Source/MutableRuntime/Internal/MuR/ProgramCache.h b/Engine/Plugins/Mutable/Source/MutableRuntime/Internal/MuR/ProgramCache.h
--- a/Engine/Plugins/Mutable/Source/MutableRuntime/Internal/MuR/ProgramCache.h
+++ b/Engine/Plugins/Mutable/Source/MutableRuntime/Internal/MuR/ProgramCache.h
@@ -100,6 +100,8 @@ namespace UE::Mutable::Private
 
 			ImageDesc = 1 << 3,
 			Void      = 1 << 4,
+			
+			ForceClearMarkedToKeep = 1 << 5,
 		};
 
 		enum class EAcquireSetResult : uint8
@@ -125,6 +127,7 @@ namespace UE::Mutable::Private
 		bool IsSet(FCacheAddress Address);
 		void UpdateHitCount(FCacheAddress Address, int8 Differential);
 		void LockAddress(FCacheAddress Address);
+		void MarkAddressToKeep(FCacheAddress Address);
 		void SetAborted(FCacheAddress Address);
 		bool IsAborted(FCacheAddress Address);
 
diff --git a/Engine/Plugins/Mutable/Source/MutableRuntime/Private/MuR/CodeRunner.cpp b/Engine/Plugins/Mutable/Source/MutableRuntime/Private/MuR/CodeRunner.cpp
--- a/Engine/Plugins/Mutable/Source/MutableRuntime/Private/MuR/CodeRunner.cpp
+++ b/Engine/Plugins/Mutable/Source/MutableRuntime/Private/MuR/CodeRunner.cpp
@@ -1142,7 +1142,7 @@ namespace UE::Mutable::Private
 				{
 					AddOp(FScheduledOp(Item.At, Item, 1),
 						FScheduledOp(Args.Mesh, Item),
-						FScheduledOp(Args.Layout, Item));
+						FScheduledOp::FromOpAndOptions(Args.Layout, Item, 0));
 				}
 				break;
 			}
@@ -1158,7 +1158,8 @@ namespace UE::Mutable::Private
 				else
 				{
 					TManagedPtr<const FMesh> Base = LoadMesh(FCacheAddress(Args.Mesh, Item));
-					TManagedPtr<const FLayout> Layout = LoadLayout(FCacheAddress(Args.Layout, Item));
+					TManagedPtr<const FLayout> Layout = LoadLayout(
+							FCacheAddress(FScheduledOp::FromOpAndOptions(Args.Layout, Item, 0)));
 
 					if (Base)
 					{
@@ -1203,7 +1204,7 @@ namespace UE::Mutable::Private
 				{
 					AddOp(FScheduledOp(Item.At, Item, 1),
 						FScheduledOp(Args.Mesh, Item),
-						FScheduledOp(Args.Layout, Item));
+						FScheduledOp::FromOpAndOptions(Args.Layout, Item, 0));
 				}
 				break;
 			}
@@ -1219,7 +1220,8 @@ namespace UE::Mutable::Private
 				else
 				{
 					TManagedPtr<const FMesh> Base = LoadMesh(FCacheAddress(Args.Mesh, Item));
-					TManagedPtr<const FLayout> Layout = LoadLayout(FCacheAddress(Args.Layout, Item));
+					TManagedPtr<const FLayout> Layout = LoadLayout(
+							FCacheAddress(FScheduledOp::FromOpAndOptions(Args.Layout, Item, 0)));
 
 					if (Base && Layout)
 					{
@@ -1561,7 +1563,7 @@ namespace UE::Mutable::Private
 						FScheduledOp(Args.Source, Item),
 						FScheduledOp(Args.UVSource, Item),
 						FScheduledOp(Args.MaskImage, Item),
-						FScheduledOp(Args.MaskLayout, Item));
+						FScheduledOp::FromOpAndOptions(Args.MaskLayout, Item, 0));
 				}
 				break;
 			}
@@ -1579,7 +1581,8 @@ namespace UE::Mutable::Private
 					TManagedPtr<const FMesh> Source = LoadMesh(FCacheAddress(Args.Source, Item));
 					TManagedPtr<const FMesh> UVSource = LoadMesh(FCacheAddress(Args.UVSource, Item));
 					TManagedPtr<const FImage> MaskImage = LoadImage(FCacheAddress(Args.MaskImage, Item));
-					TManagedPtr<const FLayout> MaskLayout = LoadLayout(FCacheAddress(Args.MaskLayout, Item));
+					TManagedPtr<const FLayout> MaskLayout = LoadLayout(
+							FCacheAddress(FScheduledOp::FromOpAndOptions(Args.MaskLayout, Item, 0)));
 
 					// Only if both are valid.
 					if (Source && MaskImage)
@@ -3982,7 +3985,9 @@ namespace UE::Mutable::Private
 			{
 			case 0:
 			{
-				AddOp(FScheduledOp(Item.At, Item, 1), FScheduledOp::FromOpAndOptions(Args.Layout, Item, 0), FScheduledOp::FromOpAndOptions(Args.SourceLayout, Item, 0));
+				AddOp(FScheduledOp(Item.At, Item, 1), 
+						FScheduledOp::FromOpAndOptions(Args.Layout, Item, 0), 
+						FScheduledOp::FromOpAndOptions(Args.SourceLayout, Item, 0));
 				break;
 			}
 			case 1:
@@ -3990,11 +3995,11 @@ namespace UE::Mutable::Private
 				MUTABLE_CPUPROFILER_SCOPE(RunCode_ImageMultiCompose_1);
 
 				TManagedPtr<const FLayout> ComposeLayout =
-					LoadLayout(FCacheAddress(Args.Layout, FScheduledOp::FromOpAndOptions(Args.Layout, Item, 0)));
+					LoadLayout(FCacheAddress(FScheduledOp::FromOpAndOptions(Args.Layout, Item, 0)));
 
 
 				TManagedPtr<const FLayout> SourceLayout =
-					LoadLayout(FCacheAddress(Args.SourceLayout, FScheduledOp::FromOpAndOptions(Args.SourceLayout, Item, 0)));
+					LoadLayout(FCacheAddress(FScheduledOp::FromOpAndOptions(Args.SourceLayout, Item, 0)));
 
 				bool bHasBlocksToCompose = false;
 				for (const FLayoutBlock& Block : SourceLayout->Blocks)
@@ -6910,6 +6915,8 @@ namespace UE::Mutable::Private
     {
         MUTABLE_CPUPROFILER_SCOPE(RunCode_Layout);
 
+		check(Item.ExecutionOptions == 0);
+
 		EOpType type = Program.GetOpType(Item.At);
         switch (type)
         {
@@ -6937,6 +6944,7 @@ namespace UE::Mutable::Private
 
             case 1:
             {
+				MUTABLE_CPUPROFILER_SCOPE(RunCode_Layout_LA_MERGE)
                 TManagedPtr<const FLayout> pA = LoadLayout( FCacheAddress(Args.Base,Item) );
                 TManagedPtr<const FLayout> pB = LoadLayout( FCacheAddress(Args.Added,Item) );
 
@@ -6979,6 +6987,7 @@ namespace UE::Mutable::Private
 
             case 1:
             {
+				MUTABLE_CPUPROFILER_SCOPE(RunCode_Layout_LA_PACK)
                 TManagedPtr<const FLayout> Source = LoadLayout( FCacheAddress(Args.Source,Item) );
 
 				TManagedPtr<FLayout> Result;
@@ -7002,6 +7011,7 @@ namespace UE::Mutable::Private
 
 		case EOpType::LA_FROMMESH:
 		{
+
 			FOperation::LayoutFromMeshArgs Args = Program.GetOpArgs<FOperation::LayoutFromMeshArgs>(Item.At);
 
 			constexpr uint8 MeshContentFilter = (uint8)(EMeshContentFlags::AllFlags);
@@ -7014,11 +7024,19 @@ namespace UE::Mutable::Private
 
 			case 1:
 			{
+				MUTABLE_CPUPROFILER_SCOPE(RunCode_Layout_LA_FROMMESH)
+
 				TManagedPtr<const FMesh> Mesh = LoadMesh(
 						FScheduledOp::FromOpAndOptions(Args.Mesh, Item, MeshContentFilter));
 
 				TManagedPtr<const FLayout> Result = LayoutFromMesh_RemoveBlocks(Mesh.Get(), Args.LayoutIndex);
 				StoreLayout(Item, Result);
+
+				if (Result)
+				{
+					GetMemory().MarkAddressToKeep(Item);
+				}
+
 				break;
 			}
 
@@ -7042,6 +7060,7 @@ namespace UE::Mutable::Private
 
 			case 1:
 			{
+				MUTABLE_CPUPROFILER_SCOPE(RunCode_Layout_LA_REMOVEBLOCKS)
 				TManagedPtr<const FLayout> Source = LoadLayout(FCacheAddress(Args.Source, Item));
 				TManagedPtr<const FLayout> ReferenceLayout = LoadLayout(FCacheAddress(Args.ReferenceLayout, Item));
 
diff --git a/Engine/Plugins/Mutable/Source/MutableRuntime/Private/MuR/ProgramCache.cpp b/Engine/Plugins/Mutable/Source/MutableRuntime/Private/MuR/ProgramCache.cpp
--- a/Engine/Plugins/Mutable/Source/MutableRuntime/Private/MuR/ProgramCache.cpp
+++ b/Engine/Plugins/Mutable/Source/MutableRuntime/Private/MuR/ProgramCache.cpp
@@ -17,9 +17,11 @@ namespace ProgramCacheInternals
 		None     = 0,
 		IsSet    = 1 << 0,
 		IsLocked = 1 << 1,
-		Aborted  = 1 << 2,
+		IsMarkedToKeep = 1 << 2,
 
-		AllocCounted  = 1 << 3,
+		Aborted  = 1 << 3,
+
+		AllocCounted  = 1 << 4,
 	};
 
 	ENUM_CLASS_FLAGS(EOpExecDataFlags);
@@ -113,6 +115,7 @@ namespace ProgramCacheInternals
 		{
 			return static_cast<CacheType*>(Cache)->RollbackAcquireResult(Address);
 		}
+
 		static void Lock(void* Cache)
 		{
 			return static_cast<CacheType*>(Cache)->Lock();
@@ -344,6 +347,12 @@ namespace ProgramCacheInternals
 
 				EOpExecDataFlags AddressFlags = Data[Index].Flags;
 
+				if (EnumHasAnyFlags(AddressFlags, EOpExecDataFlags::IsMarkedToKeep) && 
+				    !EnumHasAnyFlags(ClearFlags, FProgramCache::EClearFlags::ForceClearMarkedToKeep))
+				{
+					continue;
+				}	
+
 				bool bIsLockedAddress = EnumHasAnyFlags(AddressFlags, EOpExecDataFlags::IsLocked);
 				if constexpr (CIsWeakable<T>)
 				{
@@ -1204,6 +1213,18 @@ namespace ProgramCacheInternals
 		GetTypeErasedCacheForAddress(GetInternals(), Address).UpdateAddressState(Address, Update);
 	}
 
+	void FProgramCache::MarkAddressToKeep(FCacheAddress Address)
+	{
+		using namespace ProgramCacheInternals;
+		constexpr FAddressStateUpdate Update
+		{
+			.AddedFlags        = EOpExecDataFlags::IsMarkedToKeep,
+			.OpHitDifferential = 0,
+		};
+
+		GetTypeErasedCacheForAddress(GetInternals(), Address).UpdateAddressState(Address, Update);
+	}
+
 	void FProgramCache::SetAborted(FCacheAddress Address)
 	{
 		using namespace ProgramCacheInternals;
