[mutable] Fixed crash when the RunCode operation was aborted (material parameter could not be found).
- Removed redundant task nesting for the GetImage task.

#rnx
#Synced-CL 52611640

diff --git a/Engine/Plugins/Mutable/Source/CustomizableObject/Private/MuCO/CustomizableObjectMeshUpdate.cpp b/Engine/Plugins/Mutable/Source/CustomizableObject/Private/MuCO/CustomizableObjectMeshUpdate.cpp
--- a/Engine/Plugins/Mutable/Source/CustomizableObject/Private/MuCO/CustomizableObjectMeshUpdate.cpp
+++ b/Engine/Plugins/Mutable/Source/CustomizableObject/Private/MuCO/CustomizableObjectMeshUpdate.cpp
@@ -178,11 +178,13 @@ namespace impl
 		const UE::Mutable::Private::FMeshId& MeshID = UpdateContext->MeshContext->MeshIDs[LODIndex];
 
 		constexpr UE::Mutable::Private::EMeshContentFlags MeshContentFilter = UE::Mutable::Private::EMeshContentFlags::AllFlags;
-		UE::Tasks::TTask<UE::Mutable::Private::TManagedPtr<const UE::Mutable::Private::FMesh>> GetMeshTask = UpdateContext->MeshContext->System->GetMesh({}, UpdateContext->MeshContext->LiveInstance, MeshID, MeshContentFilter);
+		UE::Tasks::TTask<UE::Mutable::Private::FGetMeshResult> GetMeshTask = UpdateContext->MeshContext->System->GetMesh({}, UpdateContext->MeshContext->LiveInstance, MeshID, MeshContentFilter);
 
 		UE::Tasks::AddNested(UE::Tasks::Launch(TEXT("Task_MutableGetMeshes_GetMesh_Post"), [=]() mutable
 			{
-				UpdateContext->Meshes[LODIndex] = GetMeshTask.GetResult();
+				const UE::Mutable::Private::FGetMeshResult Result = GetMeshTask.GetResult();
+				check(Result.MutableMesh);
+				UpdateContext->Meshes[LODIndex] = Result.MutableMesh;
 
 				Task_Mutable_UpdateMesh_Loop(UpdateContext, Task, LODIndex + 1);
 			},
diff --git a/Engine/Plugins/Mutable/Source/CustomizableObject/Private/MuCO/CustomizableObjectMipDataProvider.cpp b/Engine/Plugins/Mutable/Source/CustomizableObject/Private/MuCO/CustomizableObjectMipDataProvider.cpp
--- a/Engine/Plugins/Mutable/Source/CustomizableObject/Private/MuCO/CustomizableObjectMipDataProvider.cpp
+++ b/Engine/Plugins/Mutable/Source/CustomizableObject/Private/MuCO/CustomizableObjectMipDataProvider.cpp
@@ -159,24 +159,23 @@ namespace Impl
 
 		const FMutableImageReference& ImageRef = OperationData->RequestedImage;
 		
-		UE::Tasks::TTask<UE::Mutable::Private::TManagedPtr<const UE::Mutable::Private::FImage>> GetImageTask = System->GetImage({}, LiveInstance, ImageRef.ImageID, ImageRef.BaseMip + OperationData->MipsToSkip);
+		UE::Tasks::TTask<UE::Mutable::Private::FGetImageResult> GetImageTask = System->GetImage({}, LiveInstance, ImageRef.ImageID, ImageRef.BaseMip + OperationData->MipsToSkip);
 
 		UE::Tasks::AddNested(UE::Tasks::Launch(TEXT("MipDataProvider_EndUpdateImagesTask"),
 				[System, OperationData, StartTime, GetImageTask, EndUpdateImage]() mutable
 				{
 					check(GetImageTask.IsCompleted());
 
-					UE::Mutable::Private::TManagedPtr<const UE::Mutable::Private::FImage> ResultImage = GetImageTask.GetResult();
-
+					const UE::Mutable::Private::FGetImageResult Result = GetImageTask.GetResult();
+						
+					UE::Mutable::Private::TManagedPtr<const UE::Mutable::Private::FImage> ResultImage = Result.MutableImage;
 					check(ResultImage);
 
 					const int32 FullMipCount = ResultImage->GetMipmapCount(ResultImage->GetSizeX(), ResultImage->GetSizeY());
 					const int32 RealMipCount = ResultImage->GetLODCount();
 
 					// Did we fail to generate the entire mipchain (if we have mips at all)?
-					bool bForceMipchain = (RealMipCount != 1) && (RealMipCount != FullMipCount);
-
-					if (bForceMipchain)
+					if (bool bForceMipChain = (RealMipCount != 1) && (RealMipCount != FullMipCount))
 					{
 						MUTABLE_CPUPROFILER_SCOPE(GetImage_MipFix);
 
diff --git a/Engine/Plugins/Mutable/Source/CustomizableObject/Private/MuCO/CustomizableObjectSystem.cpp b/Engine/Plugins/Mutable/Source/CustomizableObject/Private/MuCO/CustomizableObjectSystem.cpp
--- a/Engine/Plugins/Mutable/Source/CustomizableObject/Private/MuCO/CustomizableObjectSystem.cpp
+++ b/Engine/Plugins/Mutable/Source/CustomizableObject/Private/MuCO/CustomizableObjectSystem.cpp
@@ -1914,13 +1914,21 @@ namespace impl
 		{
 			TSharedRef<FInstanceUpdateData::FMaterial> Material = Context->InstanceUpdateData.Materials[MaterialIndex];
 			
-			UE::Tasks::TTask<UE::Mutable::Private::TManagedPtr<const UE::Mutable::Private::FMaterial>> GetMaterialTask = Context->MutableSystem->GetMaterial(LastTask, Context->LiveInstance.ToSharedRef(), Material->MaterialId);
-			UE::Tasks::AddNested(GetMaterialTask);
-			LastTask = GetMaterialTask;
+			UE::Tasks::TTask<UE::Mutable::Private::FGetMaterialResult> GetMaterialTask = Context->MutableSystem->GetMaterial(LastTask, Context->LiveInstance.ToSharedRef(), Material->MaterialId);
 			
 			LastTask = UE::Tasks::Launch(TEXT("GetMaterialResult"), [Context, ModelResources, Material, GetMaterialTask]() mutable
 			{
-				Material->Material = GetMaterialTask.GetResult();
+				const UE::Mutable::Private::FGetMaterialResult Result = GetMaterialTask.GetResult();
+				check(Result.MutableMaterial);
+					
+				// The material could not be found so just skip the processing of it
+				if (!Result.bWasOperationSuccessful)
+				{
+					Context->bIsUpdateAborted = true;
+					return;
+				}
+					
+				Material->Material = Result.MutableMaterial;
 
 				// Images
 				Material->FirstImage = Context->InstanceUpdateData.Images.Num();
@@ -2043,13 +2051,15 @@ namespace impl
 
 			Component->SkeletalMeshId = Instance->GetSkeletalMeshId(ComponentIndex);
 
-			UE::Tasks::TTask<UE::Mutable::Private::TManagedPtr<const UE::Mutable::Private::FSkeletalMesh>> GetSkeletalMeshTask = OperationData->MutableSystem->GetSkeletalMesh(LastTask, OperationData->LiveInstance.ToSharedRef(), Component->SkeletalMeshId, UE::Mutable::Private::EMeshContentFlags::AllFlags);
+			UE::Tasks::TTask<UE::Mutable::Private::FGetSkeletalMeshResult> GetSkeletalMeshTask = OperationData->MutableSystem->GetSkeletalMesh(LastTask, OperationData->LiveInstance.ToSharedRef(), Component->SkeletalMeshId, UE::Mutable::Private::EMeshContentFlags::AllFlags);
 			UE::Tasks::AddNested(GetSkeletalMeshTask);
 			LastTask = GetSkeletalMeshTask;
 			
 			UE::Tasks::AddNested(UE::Tasks::Launch(TEXT("GetSkeletalMeshResult"), [=]() mutable
 			{
-				Component->SkeletalMesh = GetSkeletalMeshTask.GetResult();
+				const UE::Mutable::Private::FGetSkeletalMeshResult Result = GetSkeletalMeshTask.GetResult();
+				check(Result.MutableSkeletalMesh);
+				Component->SkeletalMesh = Result.MutableSkeletalMesh;
 			},
 			GetSkeletalMeshTask,
 			LowLevelTasks::ETaskPriority::Inherit));
@@ -2129,14 +2139,16 @@ namespace impl
 					}
 				}
 				
-				UE::Tasks::TTask<UE::Mutable::Private::TManagedPtr<const UE::Mutable::Private::FMesh>> GetMeshTask = OperationData->MutableSystem->GetMesh(LastTask, OperationData->LiveInstance.ToSharedRef(), LOD->MeshId, MeshContentFilter);
+				UE::Tasks::TTask<UE::Mutable::Private::FGetMeshResult> GetMeshTask = OperationData->MutableSystem->GetMesh(LastTask, OperationData->LiveInstance.ToSharedRef(), LOD->MeshId, MeshContentFilter);
 				UE::Tasks::AddNested(GetMeshTask);
 				LastTask = GetMeshTask;
 				
 				UE::Tasks::AddNested(UE::Tasks::Launch(TEXT("GetMeshResult"),
 				[LOD, GetMeshTask]() mutable
 				{
-					LOD->Mesh = GetMeshTask.GetResult();
+					const UE::Mutable::Private::FGetMeshResult Result = GetMeshTask.GetResult();
+					check(Result.MutableMesh);
+					LOD->Mesh = Result.MutableMesh;
 				},
 				GetMeshTask,
 				UE::Tasks::ETaskPriority::Inherit));
@@ -2233,13 +2245,15 @@ namespace impl
 
 			if (Image->bIsPassThrough)
 			{
-				UE::Tasks::TTask<UE::Mutable::Private::TManagedPtr<const UE::Mutable::Private::FImage>> GetImageTask = OperationData->MutableSystem->GetImage(LastTask, OperationData->LiveInstance.ToSharedRef(), Image->ImageID, 0);
+				UE::Tasks::TTask<UE::Mutable::Private::FGetImageResult> GetImageTask = OperationData->MutableSystem->GetImage(LastTask, OperationData->LiveInstance.ToSharedRef(), Image->ImageID, 0);
 				UE::Tasks::AddNested(GetImageTask);
 				LastTask = GetImageTask;
 					
 				UE::Tasks::AddNested(UE::Tasks::Launch(TEXT("GetImageResult"), [Image, GetImageTask]() mutable
 				{
-					Image->Image = GetImageTask.GetResult();
+					const UE::Mutable::Private::FGetImageResult Result = GetImageTask.GetResult();
+					check(Result.MutableImage);
+					Image->Image = Result.MutableImage;
 					check(Image->Image->IsReference());
 				},
 				GetImageTask,
@@ -2284,25 +2298,25 @@ namespace impl
 				const int32 MipSizeX = FMath::Max(Image->FullImageSizeX >> MipsToSkip, 1);
 				const int32 MipSizeY = FMath::Max(Image->FullImageSizeY >> MipsToSkip, 1);
 
-				UE::Tasks::TTask<UE::Mutable::Private::TManagedPtr<const UE::Mutable::Private::FImage>> GetImageTask = OperationData->MutableSystem->GetImage(LastTask, OperationData->LiveInstance.ToSharedRef(), Image->ImageID, Image->BaseMip + MipsToSkip);
+				UE::Tasks::TTask<UE::Mutable::Private::FGetImageResult> GetImageTask = OperationData->MutableSystem->GetImage(LastTask, OperationData->LiveInstance.ToSharedRef(), Image->ImageID, Image->BaseMip + MipsToSkip);
 				UE::Tasks::AddNested(GetImageTask);
 				LastTask = GetImageTask;
 
 				UE::Tasks::AddNested(UE::Tasks::Launch(TEXT("GetImageResult"), [Image, GetImageTask, MipSizeX, MipSizeY, FullLODCount, MipsToSkip]() mutable
 				{
-					Image->Image = GetImageTask.GetResult();
+					const UE::Mutable::Private::FGetImageResult Result = GetImageTask.GetResult();
+					Image->Image = Result.MutableImage;
 					check(Image->Image);
 
 					// We should have generated exactly this size.
-					const bool bSizeMissmatch = Image->Image->GetSizeX() != MipSizeX || Image->Image->GetSizeY() != MipSizeY;
-					if (bSizeMissmatch)
+					if (const bool bSizeMissmatch = Image->Image->GetSizeX() != MipSizeX || Image->Image->GetSizeY() != MipSizeY)
 					{
 						// Generate a correctly-sized but empty image instead, to avoid crashes.
 						UE_LOGF(LogMutable, Warning, "Mutable generated a wrongly-sized image %ls.", *Image->ImageID.ToString());
 						Image->Image = UE::Mutable::Private::MakeManaged<UE::Mutable::Private::FImage>(MipSizeX, MipSizeY, FullLODCount - MipsToSkip, Image->Image->GetFormat(), UE::Mutable::Private::EInitializationType::Black);
 					}
 
-					// We need one mip or the complete chain. Otherwise there was a bug.
+					// We need one mip or the complete chain. Otherwise, there was a bug.
 					const int32 FullMipCount = Image->Image->GetMipmapCount(Image->Image->GetSizeX(), Image->Image->GetSizeY());
 					const int32 RealMipCount = Image->Image->GetLODCount();
 
diff --git a/Engine/Plugins/Mutable/Source/MutableRuntime/Internal/MuR/System.h b/Engine/Plugins/Mutable/Source/MutableRuntime/Internal/MuR/System.h
--- a/Engine/Plugins/Mutable/Source/MutableRuntime/Internal/MuR/System.h
+++ b/Engine/Plugins/Mutable/Source/MutableRuntime/Internal/MuR/System.h
@@ -247,6 +247,34 @@ namespace UE::Mutable::Private
 	};
 	
 	
+	struct FGetImageResult
+	{
+		TManagedPtr<const FImage> MutableImage;
+		
+		bool bWasOperationSuccessful = true;
+	};
+	
+	struct FGetMaterialResult
+	{
+		TManagedPtr<const FMaterial> MutableMaterial;
+		
+		bool bWasOperationSuccessful = true;
+	};
+	
+	struct FGetSkeletalMeshResult
+	{
+		TManagedPtr<const FSkeletalMesh> MutableSkeletalMesh;
+		
+		bool bWasOperationSuccessful = true;
+	};
+		
+	struct FGetMeshResult
+	{
+		TManagedPtr<const FMesh> MutableMesh;
+		
+		bool bWasOperationSuccessful = true;
+	};
+	
     /** Main system class to load models and build instances. */
 	class FSystem
 	{
@@ -319,12 +347,12 @@ namespace UE::Mutable::Private
 		//! \param MipsToSkip Number of mips to skip compared from the full image.
 		//! If 0, all mip levels will be generated. If more levels than possible to discard are specified, 
 		//! the image will still contain a minimum number of mips specified at model compile time.
-		UE_API UE::Tasks::TTask<TManagedPtr<const FImage>> GetImage(UE::Tasks::FTask Prerequisite, const TSharedRef<FLiveInstance>& LiveInstance, const FImageId& ImageId, int32 MipsToSkip = 0);
+		UE_API UE::Tasks::TTask<FGetImageResult> GetImage(UE::Tasks::FTask Prerequisite, const TSharedRef<FLiveInstance>& LiveInstance, const FImageId& ImageId, int32 MipsToSkip = 0);
 
 		//! Only valid between BeginUpdate and EndUpdate
-		UE_API UE::Tasks::TTask<TManagedPtr<const FMesh>> GetMesh(UE::Tasks::FTask Prerequisite, const TSharedRef<FLiveInstance>& LiveInstance, const FMeshId& MeshId, EMeshContentFlags MeshContentFilter = EMeshContentFlags::AllFlags);
+		UE_API UE::Tasks::TTask<FGetMeshResult> GetMesh(UE::Tasks::FTask Prerequisite, const TSharedRef<FLiveInstance>& LiveInstance, const FMeshId& MeshId, EMeshContentFlags MeshContentFilter = EMeshContentFlags::AllFlags);
 		
-		UE_API UE::Tasks::TTask<TManagedPtr<const FSkeletalMesh>> GetSkeletalMesh(UE::Tasks::FTask Prerequisite, const TSharedRef<FLiveInstance>& LiveInstance, const FSkeletalMeshId& SkeletalMeshId, EMeshContentFlags MeshContentFilter);
+		UE_API UE::Tasks::TTask<FGetSkeletalMeshResult> GetSkeletalMesh(UE::Tasks::FTask Prerequisite, const TSharedRef<FLiveInstance>& LiveInstance, const FSkeletalMeshId& SkeletalMeshId, EMeshContentFlags MeshContentFilter);
 
 		//! Only valid between BeginUpdate and EndUpdate
 		//! Calculate the description of an image, without generating it.
@@ -339,7 +367,7 @@ namespace UE::Mutable::Private
         //! Only valid between BeginUpdate and EndUpdate
 		UE_API TManagedPtr<const FMesh> GetMeshInline(const TSharedRef<FLiveInstance>& LiveInstance, const FMeshId& MeshId, EMeshContentFlags MeshContentFilter = EMeshContentFlags::AllFlags);
 		
-		UE_API UE::Tasks::TTask<TManagedPtr<const FMaterial>> GetMaterial(UE::Tasks::FTask Prerequisite, const TSharedRef<FLiveInstance>& LiveInstance, const FMaterialId& MaterialId);
+		UE_API UE::Tasks::TTask<FGetMaterialResult> GetMaterial(UE::Tasks::FTask Prerequisite, const TSharedRef<FLiveInstance>& LiveInstance, const FMaterialId& MaterialId);
 		
         //! Invalidate and free the last Instance data returned by a call to BeginUpdate with
         //! the same instance index. After a call to this method, that Instance cannot be used any
@@ -382,11 +410,6 @@ namespace UE::Mutable::Private
 
         //! Data streaming interface, if any.
 		TSharedPtr<FModelReader> StreamInterface;
-    	
-		/** This flag is turned on when a streaming error or similar happens. Results are not usable.
-		* This should only happen in-editor.
-		*/
-		bool bUnrecoverableError = false;
 		
 		FWorkingMemoryManager WorkingMemoryManager;
 
@@ -395,7 +418,10 @@ namespace UE::Mutable::Private
 		
 		MUTABLERUNTIME_API static bool CheckForUpdatedParameters(const TSharedRef<FLiveInstance>& LiveInstance, const TSharedPtr<const FParameters>& InOldParameters);
 		
-		void RunCode(
+		/** Runs the mutable code 
+		 * @return true if the execution was successful and false otherwise
+		 */
+		bool RunCode(
 			const TSharedRef<FLiveInstance>& InLiveInstance,
 			FOperation::ADDRESS,
 			uint32 LODs = AllLODs,
diff --git a/Engine/Plugins/Mutable/Source/MutableRuntime/Private/MuR/CodeRunner.cpp b/Engine/Plugins/Mutable/Source/MutableRuntime/Private/MuR/CodeRunner.cpp
--- a/Engine/Plugins/Mutable/Source/MutableRuntime/Private/MuR/CodeRunner.cpp
+++ b/Engine/Plugins/Mutable/Source/MutableRuntime/Private/MuR/CodeRunner.cpp
@@ -184,6 +184,10 @@ namespace UE::Mutable::Private
 	{
 		MUTABLE_CPUPROFILER_SCOPE(CodeRunner_Create)
 		
+		RootOPAddress = At;
+		RootOpExecutionOptions = ExecutionOptions;
+		RootOpType = Type;
+		
 		if (Type == FScheduledOp::EType::ImageDesc)
 		{
 			ImageDescConstantImages.Reserve(32);
diff --git a/Engine/Plugins/Mutable/Source/MutableRuntime/Private/MuR/CodeRunner.h b/Engine/Plugins/Mutable/Source/MutableRuntime/Private/MuR/CodeRunner.h
--- a/Engine/Plugins/Mutable/Source/MutableRuntime/Private/MuR/CodeRunner.h
+++ b/Engine/Plugins/Mutable/Source/MutableRuntime/Private/MuR/CodeRunner.h
@@ -644,6 +644,14 @@ namespace  UE::Mutable::Private
 
 		/** */
 		void LaunchIssuedTask(const TSharedPtr<FIssuedTask>& TaskToIssue, bool& bOutFailed);
+    	
+    private:
+    	
+    	FOperation::ADDRESS RootOPAddress = 0;
+    	
+    	uint8 RootOpExecutionOptions;
+    	
+    	FScheduledOp::EType RootOpType;
     };
 
 
diff --git a/Engine/Plugins/Mutable/Source/MutableRuntime/Private/MuR/CodeRunnerTasked.cpp b/Engine/Plugins/Mutable/Source/MutableRuntime/Private/MuR/CodeRunnerTasked.cpp
--- a/Engine/Plugins/Mutable/Source/MutableRuntime/Private/MuR/CodeRunnerTasked.cpp
+++ b/Engine/Plugins/Mutable/Source/MutableRuntime/Private/MuR/CodeRunnerTasked.cpp
@@ -161,10 +161,23 @@ namespace UE::Mutable::Private
 		ImageDescResults.Reset();
 		ImageDescConstantImages.Reset();
 		
-		bool bProfile = false;
-		
+		constexpr bool bProfile = false;
 		TUniquePtr<FProfileContext> ProfileContext = bProfile ? MakeUnique<FProfileContext>() : nullptr;
 		Run(MoveTemp(ProfileContext), bForceInlineExecution);
+		
+		{
+			FCacheAddress CacheAddress = FCacheAddress(RootOPAddress, 0, RootOpExecutionOptions, RootOpType);
+
+			// In case of aborted the value will not be read, decrease root address hit count
+			// to balance the ophit added at runner creation.
+			const bool bAborted = LiveInstance->Cache->IsAborted(CacheAddress);
+			if (bAborted)
+			{
+				LiveInstance->Cache->UpdateHitCount(CacheAddress, -1);
+			}
+		
+			bUnrecoverableError = bUnrecoverableError || bAborted;
+		}
 	}
 
 	void CodeRunner::LaunchWaitForOthersTask(Tasks::FTaskEvent RunnerRestartEvent, uint64 ReentryCounter)
diff --git a/Engine/Plugins/Mutable/Source/MutableRuntime/Private/MuR/System.cpp b/Engine/Plugins/Mutable/Source/MutableRuntime/Private/MuR/System.cpp
--- a/Engine/Plugins/Mutable/Source/MutableRuntime/Private/MuR/System.cpp
+++ b/Engine/Plugins/Mutable/Source/MutableRuntime/Private/MuR/System.cpp
@@ -255,8 +255,7 @@ namespace UE::Mutable::Private
 		{
 			const FOperation::ADDRESS RootAt = Program.States[LiveInstanceState].Root;
 			
-			RunCode(LiveInstance, RootAt, InLodMask);
-			if (!bUnrecoverableError)
+			if (RunCode(LiveInstance, RootAt, InLodMask))
 			{
 				Result = LiveInstance->Cache->LoadInstance(FCacheAddress(RootAt, 0, 0));
 			}
@@ -301,7 +300,7 @@ namespace UE::Mutable::Private
 		return Result;
 	}
 
-	UE::Tasks::TTask<TManagedPtr<const FImage>> FSystem::GetImage(UE::Tasks::FTask Prerequisite, const TSharedRef<FLiveInstance>& LiveInstance, const FImageId& ImageId, int32 MipsToSkip)
+	UE::Tasks::TTask<FGetImageResult> FSystem::GetImage(UE::Tasks::FTask Prerequisite, const TSharedRef<FLiveInstance>& LiveInstance, const FImageId& ImageId, int32 MipsToSkip)
 	{
 		LLM_SCOPE_BYNAME(TEXT("MutableRuntime"));
 		MUTABLE_CPUPROFILER_SCOPE(SystemGetImage);
@@ -309,9 +308,12 @@ namespace UE::Mutable::Private
     	if (!ImageId)
     	{
     		return UE::Tasks::Launch(TEXT("FSystem::GetImageInvalid"),
-			[]() -> TManagedPtr<const FImage>
+			[]() -> FGetImageResult
 			{
-				return MakeManaged<FImage>();
+				FGetImageResult Output;
+				Output.bWasOperationSuccessful = false;
+				Output.MutableImage = MakeManaged<FImage>();
+				return Output;
 			},
 			UE::Tasks::Prerequisites(Prerequisite),
 			UE::Tasks::ETaskPriority::Inherit,
@@ -324,9 +326,12 @@ namespace UE::Mutable::Private
     	if (GetOpDataType(OpType) != EDataType::Image)
     	{
     		return UE::Tasks::Launch(TEXT("FSystem::GetImageInvalid"),
-			[]() -> TManagedPtr<const FImage>
+			[]() -> FGetImageResult
 			{
-				return MakeManaged<FImage>();
+				FGetImageResult Output;
+				Output.bWasOperationSuccessful = false;
+				Output.MutableImage = MakeManaged<FImage>();
+				return Output;
 			},
 			UE::Tasks::Prerequisites(Prerequisite),
 			UE::Tasks::ETaskPriority::Inherit,
@@ -359,12 +364,12 @@ namespace UE::Mutable::Private
 		UE::Tasks::ETaskPriority::Inherit);
     	
 		return UE::Tasks::Launch(TEXT("FSystem::GetImageResultTask"),
-		[Context, System = this, RootAddress, MipsToSkip, LiveInstance]() -> TManagedPtr<const FImage>
+		[Context, System = this, RootAddress, MipsToSkip, LiveInstance]() -> FGetImageResult
 		{
 			TManagedPtr<const FImage> Result;
 
-			FCacheAddress RootCacheAddress = FCacheAddress(RootAddress, 0, MipsToSkip);
-			bool bAborted = LiveInstance->Cache->IsAborted(RootCacheAddress);
+			const FCacheAddress RootCacheAddress = FCacheAddress(RootAddress, 0, MipsToSkip);
+			const bool bAborted = LiveInstance->Cache->IsAborted(RootCacheAddress);
 			
 			if (!bAborted)
 			{
@@ -376,8 +381,11 @@ namespace UE::Mutable::Private
 				Result = MakeManaged<FImage>(16, 16, 1, EImageFormat::RGBA_UByte, EInitializationType::Black);
 			}
 			
-			System->bUnrecoverableError = bAborted;
-			return Result;
+			FGetImageResult Output;
+			Output.MutableImage = Result;
+			Output.bWasOperationSuccessful = !Context->Runner->bUnrecoverableError && !bAborted;
+				
+			return Output;
 		},
 		UE::Tasks::Prerequisites(RunnerTask),
 		UE::Tasks::ETaskPriority::Inherit,
@@ -553,7 +561,7 @@ namespace UE::Mutable::Private
 	}
 	
 
-    UE::Tasks::TTask<TManagedPtr<const FMesh>> FSystem::GetMesh(UE::Tasks::FTask Prerequisite, const TSharedRef<FLiveInstance>& LiveInstance, const FMeshId& MeshId, EMeshContentFlags MeshContentFilter)
+    UE::Tasks::TTask<FGetMeshResult> FSystem::GetMesh(UE::Tasks::FTask Prerequisite, const TSharedRef<FLiveInstance>& LiveInstance, const FMeshId& MeshId, EMeshContentFlags MeshContentFilter)
 	{
 		LLM_SCOPE_BYNAME(TEXT("MutableRuntime"));
 		MUTABLE_CPUPROFILER_SCOPE(SystemGetMesh);
@@ -561,9 +569,12 @@ namespace UE::Mutable::Private
 		if (!MeshId)
 		{
 			return UE::Tasks::Launch(TEXT("FSystem::GetMeshInvalid"),
-			[]() -> TManagedPtr<const FMesh>
+			[]() -> FGetMeshResult
 			{
-				return MakeManaged<FMesh>();
+				FGetMeshResult Output;
+				Output.MutableMesh = MakeManaged<FMesh>();
+				Output.bWasOperationSuccessful = false;
+				return Output;
 			},
 			UE::Tasks::Prerequisites(Prerequisite),
 			UE::Tasks::ETaskPriority::Inherit,
@@ -576,9 +587,12 @@ namespace UE::Mutable::Private
 		if (GetOpDataType(OpType) != EDataType::Mesh)
 		{
 			return UE::Tasks::Launch(TEXT("FSystem::GetMeshInvalid"),
-			[]() -> TManagedPtr<const FMesh>
+			[]() -> FGetMeshResult
 			{
-				return MakeManaged<FMesh>();
+				FGetMeshResult Output;
+				Output.MutableMesh = MakeManaged<FMesh>();
+				Output.bWasOperationSuccessful = false;
+				return Output;
 			},
 			UE::Tasks::Prerequisites(Prerequisite),
 			UE::Tasks::ETaskPriority::Inherit,
@@ -613,7 +627,7 @@ namespace UE::Mutable::Private
 		UE::Tasks::ETaskPriority::Inherit);
 		
 		return UE::Tasks::Launch(TEXT("FSystem::GetMeshResultTask"),
-		[Context, System = this, RootAddress, MeshContentFilter, LiveInstance]() -> TManagedPtr<const FMesh>
+		[Context, System = this, RootAddress, MeshContentFilter, LiveInstance]() -> FGetMeshResult
 		{
 			TManagedPtr<const FMesh> Result;
 	
@@ -625,9 +639,17 @@ namespace UE::Mutable::Private
 				check(LiveInstance->Cache->IsSet(RootCacheAddress));
 				Result = LiveInstance->Cache->LoadMesh(RootCacheAddress);
 			}
-			
-			System->bUnrecoverableError = bAborted;
-			return Result;
+				
+			if (!Result)
+			{
+				Result = MakeManaged<FMesh>();
+			}
+				
+			FGetMeshResult Output;
+			Output.MutableMesh = Result;
+			Output.bWasOperationSuccessful = !Context->Runner->bUnrecoverableError && !bAborted;;
+				
+			return Output;
 		},
 		UE::Tasks::Prerequisites(RunnerTask),
 		UE::Tasks::ETaskPriority::Inherit,
@@ -635,7 +657,7 @@ namespace UE::Mutable::Private
 	}
 
 
-	UE::Tasks::TTask<TManagedPtr<const FSkeletalMesh>> FSystem::GetSkeletalMesh(UE::Tasks::FTask Prerequisite, const TSharedRef<FLiveInstance>& LiveInstance, const FSkeletalMeshId& SkeletalMeshId, EMeshContentFlags SkeletalMeshContentFilter)
+	UE::Tasks::TTask<FGetSkeletalMeshResult> FSystem::GetSkeletalMesh(UE::Tasks::FTask Prerequisite, const TSharedRef<FLiveInstance>& LiveInstance, const FSkeletalMeshId& SkeletalMeshId, EMeshContentFlags SkeletalMeshContentFilter)
     {
 		LLM_SCOPE_BYNAME(TEXT("MutableRuntime"));
 		MUTABLE_CPUPROFILER_SCOPE(SystemGetSkeletalMesh);
@@ -643,9 +665,12 @@ namespace UE::Mutable::Private
 		if (!SkeletalMeshId)
 		{
 			return UE::Tasks::Launch(TEXT("FSystem::GetSkeletalMeshInvalid"),
-			[]() -> TManagedPtr<const FSkeletalMesh>
+			[]() -> FGetSkeletalMeshResult
 			{
-				return MakeManaged<FSkeletalMesh>();
+				FGetSkeletalMeshResult Output;
+				Output.MutableSkeletalMesh = MakeManaged<FSkeletalMesh>();
+				Output.bWasOperationSuccessful = false;
+				return Output;
 			},
 			UE::Tasks::Prerequisites(Prerequisite),
 			UE::Tasks::ETaskPriority::Inherit,
@@ -658,9 +683,12 @@ namespace UE::Mutable::Private
 		if (GetOpDataType(OpType) != EDataType::SkeletalMesh)
 		{
 			return UE::Tasks::Launch(TEXT("FSystem::GetSkeletalMeshInvalid"),
-				[]() -> TManagedPtr<const FSkeletalMesh>
+				[]() -> FGetSkeletalMeshResult
 				{
-					return MakeManaged<FSkeletalMesh>();
+					FGetSkeletalMeshResult Output;
+					Output.MutableSkeletalMesh = MakeManaged<FSkeletalMesh>();
+					Output.bWasOperationSuccessful = false;
+					return Output;
 				},
 				UE::Tasks::Prerequisites(Prerequisite),
 				UE::Tasks::ETaskPriority::Inherit,
@@ -694,7 +722,7 @@ namespace UE::Mutable::Private
 		UE::Tasks::ETaskPriority::Inherit);
 		
 		return UE::Tasks::Launch(TEXT("FSystem::GetSkeletalMeshResultTask"),
-		[Context, System = this, RootAddress, SkeletalMeshContentFilter, LiveInstance]() -> TManagedPtr<const FSkeletalMesh>
+		[Context, System = this, RootAddress, SkeletalMeshContentFilter, LiveInstance]() -> FGetSkeletalMeshResult
 		{
 			TManagedPtr<const FSkeletalMesh> Result;
 
@@ -706,9 +734,17 @@ namespace UE::Mutable::Private
 				check(LiveInstance->Cache->IsSet(RootCacheAddress));
 				Result = LiveInstance->Cache->LoadSkeletalMesh(RootCacheAddress);
 			}
-			
-			System->bUnrecoverableError = bAborted;
-			return Result;
+				
+			if (!Result)
+			{
+				Result = MakeManaged<FSkeletalMesh>();
+			}
+				
+			FGetSkeletalMeshResult Output;
+			Output.MutableSkeletalMesh = Result;
+			Output.bWasOperationSuccessful = !Context->Runner->bUnrecoverableError && !bAborted;
+				
+			return Output;
 		},
 		UE::Tasks::Prerequisites(RunnerTask),
 		UE::Tasks::ETaskPriority::Inherit,
@@ -716,7 +752,7 @@ namespace UE::Mutable::Private
 	}
 
 	
-	UE::Tasks::TTask<TManagedPtr<const FMaterial>> FSystem::GetMaterial(UE::Tasks::FTask Prerequisite, const TSharedRef<FLiveInstance>& LiveInstance, const FMaterialId& MaterialId)
+	UE::Tasks::TTask<FGetMaterialResult> FSystem::GetMaterial(UE::Tasks::FTask Prerequisite, const TSharedRef<FLiveInstance>& LiveInstance, const FMaterialId& MaterialId)
 	{
 		LLM_SCOPE_BYNAME(TEXT("MutableRuntime"));
 		MUTABLE_CPUPROFILER_SCOPE(SystemGetMaterial);
@@ -724,9 +760,12 @@ namespace UE::Mutable::Private
 		if (!MaterialId)
 		{
 			return UE::Tasks::Launch(TEXT("FSystem::GetMaterialInvalid"),
-			[]() -> TManagedPtr<const FMaterial>
+			[]() -> FGetMaterialResult
 			{
-				return MakeManaged<FMaterial>();
+				FGetMaterialResult Output;
+				Output.MutableMaterial = MakeManaged<FMaterial>();
+				Output.bWasOperationSuccessful = false;
+				return Output;
 			},
 			UE::Tasks::Prerequisites(Prerequisite),
 			UE::Tasks::ETaskPriority::Inherit,
@@ -739,9 +778,12 @@ namespace UE::Mutable::Private
 		if (GetOpDataType(OpType) != EDataType::Material)
 		{
 			return UE::Tasks::Launch(TEXT("FSystem::GetMaterialInvalid"),
-			[]() -> TManagedPtr<const FMaterial>
+			[]() -> FGetMaterialResult
 			{
-				return MakeManaged<FMaterial>();
+				FGetMaterialResult Output;
+				Output.MutableMaterial = MakeManaged<FMaterial>();
+				Output.bWasOperationSuccessful = false;
+				return Output;
 			},
 			UE::Tasks::Prerequisites(Prerequisite),
 			UE::Tasks::ETaskPriority::Inherit,
@@ -775,20 +817,28 @@ namespace UE::Mutable::Private
 		UE::Tasks::ETaskPriority::Inherit);
 		
 		return UE::Tasks::Launch(TEXT("FSystem::GetMaterialResultTask"),
-		[Context, System = this, RootAddress, LiveInstance]() -> TManagedPtr<const FMaterial>
+		[Context, System = this, RootAddress, LiveInstance]() -> FGetMaterialResult
 		{
 			TManagedPtr<const FMaterial> Result;
 
 			FCacheAddress RootCacheAddress = FCacheAddress(RootAddress, 0, 0);
 
-			bool bAborted = LiveInstance->Cache->IsAborted(RootCacheAddress);
+			const bool bAborted = LiveInstance->Cache->IsAborted(RootCacheAddress);
 			if (!bAborted)
 			{
 				Result = Context->Runner->LoadMaterial(RootCacheAddress);
 			}
 			
-			System->bUnrecoverableError = bAborted;
-			return Result;
+			if (!Result)
+			{
+				Result = MakeManaged<FMaterial>();
+			}
+				
+			FGetMaterialResult Output;
+			Output.MutableMaterial = Result;
+			Output.bWasOperationSuccessful = !Context->Runner->bUnrecoverableError && !bAborted;
+						
+			return Output;
 		},
 		UE::Tasks::Prerequisites(RunnerTask),
 		UE::Tasks::ETaskPriority::Inherit,
@@ -922,7 +972,7 @@ namespace UE::Mutable::Private
     }
 
 
-	void FSystem::RunCode(
+	bool FSystem::RunCode(
 		const TSharedRef<FLiveInstance>& InLiveInstance,
 		FOperation::ADDRESS InCodeRoot,
 		uint32 InLODs,
@@ -938,30 +988,20 @@ namespace UE::Mutable::Private
 			ExecutionOptions,
 			false,
 			FScheduledOp::EType::Full);
-		
+    	
 		constexpr bool bForceInlineExecution = true;
 		Runner->StartRun(bForceInlineExecution);
-	
-		FCacheAddress CacheAddress = FCacheAddress(InCodeRoot, 0, ExecutionOptions);
-		bool bAborted = InLiveInstance->Cache->IsAborted(CacheAddress);
-
-		// In case of aborted the value will not be read, decrease root address hit count
-		// to balance the ophit added at runner creation.
-		if (bAborted)
-		{
-			InLiveInstance->Cache->UpdateHitCount(CacheAddress, -1);
-		}
 
-		bUnrecoverableError = bAborted;
+    	// The runner will display an unrecoverable error if one was found during the run operation or if the cache reports that the operation was marked as aborted
+    	const bool bWasRunSuccessful = !Runner->bUnrecoverableError;
+    	return bWasRunSuccessful;
 	}
 	
 	
 	bool FSystem::BuildBool(const TSharedRef<FLiveInstance>& InLiveInstance, FOperation::ADDRESS InAddress)
     {
-    	RunCode(InLiveInstance, InAddress);
-	
     	bool bResult = false;
-    	if (!bUnrecoverableError)
+    	if (RunCode(InLiveInstance, InAddress))
     	{
     		bResult = InLiveInstance->Cache->LoadBool(FCacheAddress(InAddress, 0, 0));
     	}
@@ -972,10 +1012,8 @@ namespace UE::Mutable::Private
 	
 	int32 FSystem::BuildInt(const TSharedRef<FLiveInstance>& InLiveInstance, FOperation::ADDRESS InAddress)
 	{
-		RunCode(InLiveInstance, InAddress);
-
 		int32 Result = 0;
-		if (!bUnrecoverableError)
+		if (RunCode(InLiveInstance, InAddress))
 		{
 			Result = InLiveInstance->Cache->LoadInt(FCacheAddress(InAddress, 0, 0));
 		}
@@ -986,10 +1024,8 @@ namespace UE::Mutable::Private
 	
 	float FSystem::BuildScalar(const TSharedRef<FLiveInstance>& InLiveInstance, FOperation::ADDRESS InAddress)
 	{
-    	RunCode(InLiveInstance, InAddress);
-
     	float Result = 0.0f;		
-    	if (!bUnrecoverableError)
+    	if (RunCode(InLiveInstance, InAddress))
     	{
     		Result = InLiveInstance->Cache->LoadScalar(FCacheAddress(InAddress, 0, 0));
     	}
@@ -1005,8 +1041,7 @@ namespace UE::Mutable::Private
 		UE::Mutable::Private::EOpType OpType = InLiveInstance->Model->GetProgram().GetOpType(InAddress);
 		if (GetOpDataType(OpType) == EDataType::Color)
 		{
-			RunCode(InLiveInstance, InAddress);
-			if (!bUnrecoverableError)
+			if (RunCode(InLiveInstance, InAddress))
 			{
 				Result = InLiveInstance->Cache->LoadColor(FCacheAddress(InAddress, 0, 0));
 			}
@@ -1018,10 +1053,8 @@ namespace UE::Mutable::Private
 	
 	FProjector FSystem::BuildProjector(const TSharedRef<FLiveInstance>& InLiveInstance, FOperation::ADDRESS InAddress)
 	{
-    	RunCode(InLiveInstance, InAddress);
-
 		FProjector Result;
-		if (!bUnrecoverableError)
+		if (RunCode(InLiveInstance, InAddress))
 		{
 			Result = InLiveInstance->Cache->LoadProjector(FCacheAddress(InAddress, 0, 0));
 		}
@@ -1037,8 +1070,7 @@ namespace UE::Mutable::Private
 		UE::Mutable::Private::EOpType OpType = InLiveInstance->Model->GetProgram().GetOpType(InAddress);
 		if (GetOpDataType(OpType) == EDataType::Image)
 		{
-			RunCode(InLiveInstance, InAddress, FSystem::AllLODs, uint8(MipsToSkip));
-			if (!bUnrecoverableError)
+			if (RunCode(InLiveInstance, InAddress, FSystem::AllLODs, uint8(MipsToSkip)))
 			{
 				Result = InLiveInstance->Cache->LoadImage(FCacheAddress(InAddress, 0, MipsToSkip));
 			}			
@@ -1056,8 +1088,7 @@ namespace UE::Mutable::Private
 		if (GetOpDataType(OpType) == EDataType::Mesh)
 		{
 			const uint8 ExecutionOptions = static_cast<uint8>(MeshContentFilter);
-			RunCode(InLiveInstance, InRootAddress, FSystem::AllLODs, ExecutionOptions);
-			if (!bUnrecoverableError)
+			if (RunCode(InLiveInstance, InRootAddress, FSystem::AllLODs, ExecutionOptions))
 			{
 				Result = InLiveInstance->Cache->LoadMesh(FCacheAddress(InRootAddress, 0, ExecutionOptions));
 			}
@@ -1074,8 +1105,7 @@ namespace UE::Mutable::Private
 		UE::Mutable::Private::EOpType OpType = InLiveInstance->Model->GetProgram().GetOpType(InAddress);
 		if (GetOpDataType(OpType) == EDataType::Instance)
 		{
-			RunCode(InLiveInstance, InAddress);
-			if (!bUnrecoverableError)
+			if (RunCode(InLiveInstance, InAddress))
 			{
 				Result = InLiveInstance->Cache->LoadInstance(FCacheAddress(InAddress, 0, 0));
 			}
@@ -1094,8 +1124,7 @@ namespace UE::Mutable::Private
 			UE::Mutable::Private::EOpType OpType = InLiveInstance->Model->GetProgram().GetOpType(InAddress);
 			if (GetOpDataType(OpType) == EDataType::Layout)
 			{
-				RunCode(InLiveInstance, InAddress);
-				if (!bUnrecoverableError)
+				if (RunCode(InLiveInstance, InAddress))
 				{
 					Result = InLiveInstance->Cache->LoadLayout(FCacheAddress(InAddress, 0, 0));
 				}
@@ -1115,8 +1144,7 @@ namespace UE::Mutable::Private
 			UE::Mutable::Private::EOpType OpType = InLiveInstance->Model->GetProgram().GetOpType(InAddress);
 			if (GetOpDataType(OpType) == EDataType::String)
 			{
-				RunCode(InLiveInstance, InAddress);
-				if (!bUnrecoverableError)
+				if (RunCode(InLiveInstance, InAddress))
 				{
 					Result = InLiveInstance->Cache->LoadString(FCacheAddress(InAddress, 0, 0));
 				}
