[Mutable] Fix multicompose missing format case for uncompressed L_UByte and copy mismatch between BC1 and BC3.

- Added SubImage uncompresed format function from RGBA to L and L to RGBA.
- Fix image selection function arguments flipped arguments.
- Prevent crash when src block area is 0.

#rnx
#Synced-CL 52615216

diff --git a/Engine/Plugins/Mutable/Source/MutableRuntime/Private/MuR/ImageFormatUtils.cpp b/Engine/Plugins/Mutable/Source/MutableRuntime/Private/MuR/ImageFormatUtils.cpp
--- a/Engine/Plugins/Mutable/Source/MutableRuntime/Private/MuR/ImageFormatUtils.cpp
+++ b/Engine/Plugins/Mutable/Source/MutableRuntime/Private/MuR/ImageFormatUtils.cpp
@@ -496,6 +496,39 @@ namespace
 		}
 	}
 	
+	void SubImageVec1ToVec4_U8(FImageSize FromSize, FImageSize ToSize, FImageSize SubSize, const uint8* RESTRICT From, uint8* RESTRICT To)
+	{
+		for (uint16 Y = 0; Y < SubSize.Y; ++Y)
+		{
+			for (uint16 X = 0; X < SubSize.X; ++X)
+			{
+				uint8* ToPtr         = To   + (Y*ToSize.X   + X)*4;
+				const uint8* FromPtr = From + (Y*FromSize.X + X)*1;
+
+				// TODO: Optimize.
+				ToPtr[0] = FromPtr[0];	
+				ToPtr[1] = FromPtr[0];	
+				ToPtr[2] = FromPtr[0];
+				ToPtr[3] = 255;
+			}
+		}
+	}
+
+	void SubImageVec4ToVec1_U8(FImageSize FromSize, FImageSize ToSize, FImageSize SubSize, const uint8* RESTRICT From, uint8* RESTRICT To)
+	{
+		for (uint16 Y = 0; Y < SubSize.Y; ++Y)
+		{
+			for (uint16 X = 0; X < SubSize.X; ++X)
+			{
+				uint8* ToPtr         = To   + (Y*ToSize.X   + X)*1;
+				const uint8* FromPtr = From + (Y*FromSize.X + X)*4;
+
+				// TODO: Optimize.
+				ToPtr[0] = FromPtr[0];	
+			}
+		}
+	}
+
 
 	UncompressedFormatFuncType* SelectUncompressedFormatFunction(EImageFormat DstFormat, EImageFormat SrcFormat)
 	{
@@ -512,6 +545,10 @@ namespace
 				{
 					return SubImageVec4ToVec3_U8;
 				}
+				case EImageFormat::L_UByte:
+				{
+					return SubImageVec4ToVec1_U8;
+				}
 			}
 			break;
 		}
@@ -526,6 +563,16 @@ namespace
 			}
 			break;
 		}
+		case EImageFormat::L_UByte:
+		{
+			switch (DstFormat)
+			{
+				case EImageFormat::RGBA_UByte:
+				{
+					return SubImageVec1ToVec4_U8;
+				}
+			}
+		}
 		}
 		return UncompressedFormatFuncNoOp;
 	}
diff --git a/Engine/Plugins/Mutable/Source/MutableRuntime/Private/MuR/OpImageMultiCompose.cpp b/Engine/Plugins/Mutable/Source/MutableRuntime/Private/MuR/OpImageMultiCompose.cpp
--- a/Engine/Plugins/Mutable/Source/MutableRuntime/Private/MuR/OpImageMultiCompose.cpp
+++ b/Engine/Plugins/Mutable/Source/MutableRuntime/Private/MuR/OpImageMultiCompose.cpp
@@ -112,8 +112,8 @@ namespace OpImageMultiComposeInternal
 			.DstDecompressionFunc      = SelectDecompressionFunction(DstDecompressedFormat, Dst.GetFormat()),
 			.DstCompressionFunc        = SelectCompressionFunction(Dst.GetFormat(), DstDecompressedFormat),
 			.ResizeFunc                = SubImageResize::SelectSubImageResizeFunc(DstDecompressedFormat),
-			.SrcUncompressedFormatFunc = SubImageFormat::SelectUncompressedFormatFunction(SrcDecompressedFormat, DstDecompressedFormat),
-			.DstUncompressedFormatFunc = SubImageFormat::SelectUncompressedFormatFunction(DstDecompressedFormat, MostGenericFormat),
+			.SrcUncompressedFormatFunc = SubImageFormat::SelectUncompressedFormatFunction(DstDecompressedFormat, SrcDecompressedFormat),
+			.DstUncompressedFormatFunc = SubImageFormat::SelectUncompressedFormatFunction(MostGenericFormat, DstDecompressedFormat),
 			.bIsUncompressedReformat   = bIsUncompressedReformat,
 		};
 	}
@@ -121,23 +121,29 @@ namespace OpImageMultiComposeInternal
 	EActiveStages ComputeActiveStages(const FImage& Base, const FImage& Source, const FImageFormatInfo& FormatInfo, const FIntRect& DstRect, const FIntRect& SrcRect)
 	{
 		const bool bIsSrcCompressionRectAligned = 
-				((SrcRect.Min.X % FormatInfo.Src.PixelsPerBlockX) == 0) &
-				((SrcRect.Min.Y % FormatInfo.Src.PixelsPerBlockY) == 0) &
-				((SrcRect.Max.X % FormatInfo.Src.PixelsPerBlockX) == 0) &
+				((SrcRect.Min.X % FormatInfo.Src.PixelsPerBlockX) == 0) &&
+				((SrcRect.Min.Y % FormatInfo.Src.PixelsPerBlockY) == 0) &&
+				((SrcRect.Max.X % FormatInfo.Src.PixelsPerBlockX) == 0) &&
 				((SrcRect.Max.Y % FormatInfo.Src.PixelsPerBlockY) == 0);
 
 		const bool bIsDstCompressionRectAligned = 
-				((DstRect.Min.X % FormatInfo.Dst.PixelsPerBlockX) == 0) &
-				((DstRect.Min.Y % FormatInfo.Dst.PixelsPerBlockY) == 0) &
-				((DstRect.Max.X % FormatInfo.Dst.PixelsPerBlockX) == 0) &
+				((DstRect.Min.X % FormatInfo.Dst.PixelsPerBlockX) == 0) &&
+				((DstRect.Min.Y % FormatInfo.Dst.PixelsPerBlockY) == 0) &&
+				((DstRect.Max.X % FormatInfo.Dst.PixelsPerBlockX) == 0) &&
 				((DstRect.Max.Y % FormatInfo.Dst.PixelsPerBlockY) == 0);
 
 		const bool bUncompressedFormatMissmatch =
-				(FormatInfo.SrcDecompressedFormat != FormatInfo.DstDecompressedFormat) &
-				((FormatInfo.SrcDecompressedFormat == Source.GetFormat()) & (FormatInfo.DstDecompressedFormat == Base.GetFormat()));
+				FormatInfo.SrcDecompressedFormat != FormatInfo.DstDecompressedFormat &&
+				FormatInfo.SrcDecompressedFormat == Source.GetFormat() && 
+				FormatInfo.DstDecompressedFormat == Base.GetFormat();
+
+		const bool bCompressedFormatMismatch = 
+				Source.GetFormat() != Base.GetFormat()                 && 
+				Source.GetFormat() != FormatInfo.SrcDecompressedFormat &&
+				Base.GetFormat() != FormatInfo.DstDecompressedFormat;
 
 		const bool bNeedsResize = 
-				(DstRect.Width()  != SrcRect.Width()) |
+				(DstRect.Width()  != SrcRect.Width()) ||
 				(DstRect.Height() != SrcRect.Height());
 
 		const bool bDirectDecompression = (!bNeedsResize)                              &&
@@ -146,17 +152,19 @@ namespace OpImageMultiComposeInternal
 				(FormatInfo.SrcDecompressedFormat != Source.GetFormat());
 
 		const bool bNeedsSrcDecompression = 
-				(bNeedsResize & (FormatInfo.SrcDecompressedFormat != Source.GetFormat())) |
-				(!(bIsSrcCompressionRectAligned & bIsDstCompressionRectAligned))          |
-				(bUncompressedFormatMissmatch)                                            |
-				(bDirectDecompression);
+				(bNeedsResize && (FormatInfo.SrcDecompressedFormat != Source.GetFormat())) ||
+				(!(bIsSrcCompressionRectAligned && bIsDstCompressionRectAligned))          ||
+				(bUncompressedFormatMissmatch)                                             ||
+				(bDirectDecompression)                                                     ||
+				(bCompressedFormatMismatch);
 
 		const bool bNeedsDstDecompression =
-				(bNeedsResize & (FormatInfo.DstDecompressedFormat != Base.GetFormat())) |
-				(!(bIsSrcCompressionRectAligned & bIsDstCompressionRectAligned));
+				(bNeedsResize && (FormatInfo.DstDecompressedFormat != Base.GetFormat())) ||
+				(!bNeedsResize && bCompressedFormatMismatch)                             ||
+				(!(bIsSrcCompressionRectAligned && bIsDstCompressionRectAligned));
 
 		const bool bNeedsDstCompression =
-				(FormatInfo.DstDecompressedFormat != Base.GetFormat()) & bNeedsSrcDecompression;
+				(FormatInfo.DstDecompressedFormat != Base.GetFormat()) && bNeedsSrcDecompression;
 		
 		EActiveStages Result = EActiveStages::None;
 		EnumAddFlags(Result, bNeedsResize           ? EActiveStages::Resize           : EActiveStages::None);
@@ -490,7 +498,7 @@ namespace OpImageMultiComposeInternal
 			FIntRect DstRect = DstRects[RectIndex];
 			FIntRect SrcRect = SrcRects[RectIndex];
 
-			if (DstRect.Area() == 0)
+			if (DstRect.Area() == 0 || SrcRect.Area() == 0)
 			{
 				continue;
 			}
