[mutable] Fixed potential CodeRunner data race.
- This was caused by wrongfully moving some code that depended on the inline execution of the mutable code to another place where that was not really always the case.
- Now the code will only be executed if it is known to be an inline execution.

#rnx
#Synced-CL 52672741

diff --git a/Engine/Plugins/Mutable/Source/MutableRuntime/Internal/MuR/System.h b/Engine/Plugins/Mutable/Source/MutableRuntime/Internal/MuR/System.h
--- a/Engine/Plugins/Mutable/Source/MutableRuntime/Internal/MuR/System.h
+++ b/Engine/Plugins/Mutable/Source/MutableRuntime/Internal/MuR/System.h
@@ -421,7 +421,7 @@ namespace UE::Mutable::Private
 		/** Runs the mutable code 
 		 * @return true if the execution was successful and false otherwise
 		 */
-		bool RunCode(
+		bool RunCodeInline(
 			const TSharedRef<FLiveInstance>& InLiveInstance,
 			FOperation::ADDRESS,
 			uint32 LODs = AllLODs,
diff --git a/Engine/Plugins/Mutable/Source/MutableRuntime/Private/MuR/CodeRunnerTasked.cpp b/Engine/Plugins/Mutable/Source/MutableRuntime/Private/MuR/CodeRunnerTasked.cpp
--- a/Engine/Plugins/Mutable/Source/MutableRuntime/Private/MuR/CodeRunnerTasked.cpp
+++ b/Engine/Plugins/Mutable/Source/MutableRuntime/Private/MuR/CodeRunnerTasked.cpp
@@ -165,6 +165,7 @@ namespace UE::Mutable::Private
 		TUniquePtr<FProfileContext> ProfileContext = bProfile ? MakeUnique<FProfileContext>() : nullptr;
 		Run(MoveTemp(ProfileContext), bForceInlineExecution);
 		
+		if (bForceInlineExecution)
 		{
 			FCacheAddress CacheAddress = FCacheAddress(RootOPAddress, 0, RootOpExecutionOptions, RootOpType);
 
diff --git a/Engine/Plugins/Mutable/Source/MutableRuntime/Private/MuR/System.cpp b/Engine/Plugins/Mutable/Source/MutableRuntime/Private/MuR/System.cpp
--- a/Engine/Plugins/Mutable/Source/MutableRuntime/Private/MuR/System.cpp
+++ b/Engine/Plugins/Mutable/Source/MutableRuntime/Private/MuR/System.cpp
@@ -255,7 +255,7 @@ namespace UE::Mutable::Private
 		{
 			const FOperation::ADDRESS RootAt = Program.States[LiveInstanceState].Root;
 			
-			if (RunCode(LiveInstance, RootAt, InLodMask))
+			if (RunCodeInline(LiveInstance, RootAt, InLodMask))
 			{
 				Result = LiveInstance->Cache->LoadInstance(FCacheAddress(RootAt, 0, 0));
 			}
@@ -972,7 +972,7 @@ namespace UE::Mutable::Private
     }
 
 
-	bool FSystem::RunCode(
+	bool FSystem::RunCodeInline(
 		const TSharedRef<FLiveInstance>& InLiveInstance,
 		FOperation::ADDRESS InCodeRoot,
 		uint32 InLODs,
@@ -1001,7 +1001,7 @@ namespace UE::Mutable::Private
 	bool FSystem::BuildBool(const TSharedRef<FLiveInstance>& InLiveInstance, FOperation::ADDRESS InAddress)
     {
     	bool bResult = false;
-    	if (RunCode(InLiveInstance, InAddress))
+    	if (RunCodeInline(InLiveInstance, InAddress))
     	{
     		bResult = InLiveInstance->Cache->LoadBool(FCacheAddress(InAddress, 0, 0));
     	}
@@ -1013,7 +1013,7 @@ namespace UE::Mutable::Private
 	int32 FSystem::BuildInt(const TSharedRef<FLiveInstance>& InLiveInstance, FOperation::ADDRESS InAddress)
 	{
 		int32 Result = 0;
-		if (RunCode(InLiveInstance, InAddress))
+		if (RunCodeInline(InLiveInstance, InAddress))
 		{
 			Result = InLiveInstance->Cache->LoadInt(FCacheAddress(InAddress, 0, 0));
 		}
@@ -1025,7 +1025,7 @@ namespace UE::Mutable::Private
 	float FSystem::BuildScalar(const TSharedRef<FLiveInstance>& InLiveInstance, FOperation::ADDRESS InAddress)
 	{
     	float Result = 0.0f;		
-    	if (RunCode(InLiveInstance, InAddress))
+    	if (RunCodeInline(InLiveInstance, InAddress))
     	{
     		Result = InLiveInstance->Cache->LoadScalar(FCacheAddress(InAddress, 0, 0));
     	}
@@ -1041,7 +1041,7 @@ namespace UE::Mutable::Private
 		UE::Mutable::Private::EOpType OpType = InLiveInstance->Model->GetProgram().GetOpType(InAddress);
 		if (GetOpDataType(OpType) == EDataType::Color)
 		{
-			if (RunCode(InLiveInstance, InAddress))
+			if (RunCodeInline(InLiveInstance, InAddress))
 			{
 				Result = InLiveInstance->Cache->LoadColor(FCacheAddress(InAddress, 0, 0));
 			}
@@ -1054,7 +1054,7 @@ namespace UE::Mutable::Private
 	FProjector FSystem::BuildProjector(const TSharedRef<FLiveInstance>& InLiveInstance, FOperation::ADDRESS InAddress)
 	{
 		FProjector Result;
-		if (RunCode(InLiveInstance, InAddress))
+		if (RunCodeInline(InLiveInstance, InAddress))
 		{
 			Result = InLiveInstance->Cache->LoadProjector(FCacheAddress(InAddress, 0, 0));
 		}
@@ -1070,7 +1070,7 @@ namespace UE::Mutable::Private
 		UE::Mutable::Private::EOpType OpType = InLiveInstance->Model->GetProgram().GetOpType(InAddress);
 		if (GetOpDataType(OpType) == EDataType::Image)
 		{
-			if (RunCode(InLiveInstance, InAddress, FSystem::AllLODs, uint8(MipsToSkip)))
+			if (RunCodeInline(InLiveInstance, InAddress, FSystem::AllLODs, uint8(MipsToSkip)))
 			{
 				Result = InLiveInstance->Cache->LoadImage(FCacheAddress(InAddress, 0, MipsToSkip));
 			}			
@@ -1088,7 +1088,7 @@ namespace UE::Mutable::Private
 		if (GetOpDataType(OpType) == EDataType::Mesh)
 		{
 			const uint8 ExecutionOptions = static_cast<uint8>(MeshContentFilter);
-			if (RunCode(InLiveInstance, InRootAddress, FSystem::AllLODs, ExecutionOptions))
+			if (RunCodeInline(InLiveInstance, InRootAddress, FSystem::AllLODs, ExecutionOptions))
 			{
 				Result = InLiveInstance->Cache->LoadMesh(FCacheAddress(InRootAddress, 0, ExecutionOptions));
 			}
@@ -1105,7 +1105,7 @@ namespace UE::Mutable::Private
 		UE::Mutable::Private::EOpType OpType = InLiveInstance->Model->GetProgram().GetOpType(InAddress);
 		if (GetOpDataType(OpType) == EDataType::Instance)
 		{
-			if (RunCode(InLiveInstance, InAddress))
+			if (RunCodeInline(InLiveInstance, InAddress))
 			{
 				Result = InLiveInstance->Cache->LoadInstance(FCacheAddress(InAddress, 0, 0));
 			}
@@ -1124,7 +1124,7 @@ namespace UE::Mutable::Private
 			UE::Mutable::Private::EOpType OpType = InLiveInstance->Model->GetProgram().GetOpType(InAddress);
 			if (GetOpDataType(OpType) == EDataType::Layout)
 			{
-				if (RunCode(InLiveInstance, InAddress))
+				if (RunCodeInline(InLiveInstance, InAddress))
 				{
 					Result = InLiveInstance->Cache->LoadLayout(FCacheAddress(InAddress, 0, 0));
 				}
@@ -1144,7 +1144,7 @@ namespace UE::Mutable::Private
 			UE::Mutable::Private::EOpType OpType = InLiveInstance->Model->GetProgram().GetOpType(InAddress);
 			if (GetOpDataType(OpType) == EDataType::String)
 			{
-				if (RunCode(InLiveInstance, InAddress))
+				if (RunCodeInline(InLiveInstance, InAddress))
 				{
 					Result = InLiveInstance->Cache->LoadString(FCacheAddress(InAddress, 0, 0));
 				}
