[Mutable] Code Viewer added duplicate tag to all duplicated operations.

The original one did not have the tag.

#rnx
#Synced-CL 52761943

diff --git a/Engine/Plugins/Mutable/Source/CustomizableObjectEditor/Private/MuCOE/SMutableCodeViewer.cpp b/Engine/Plugins/Mutable/Source/CustomizableObjectEditor/Private/MuCOE/SMutableCodeViewer.cpp
--- a/Engine/Plugins/Mutable/Source/CustomizableObjectEditor/Private/MuCOE/SMutableCodeViewer.cpp
+++ b/Engine/Plugins/Mutable/Source/CustomizableObjectEditor/Private/MuCOE/SMutableCodeViewer.cpp
@@ -3174,9 +3174,19 @@ FMutableCodeTreeElement::FMutableCodeTreeElement(int32 InIndexOnTree, const int3
 	}
 	
 	IndexOnTree = InIndexOnTree;
+
 	if (InDuplicatedOf)
 	{
+        bHasDuplicate = true;
+
 		DuplicatedOf = *InDuplicatedOf;
+		check(DuplicatedOf);
+			   
+		if (!DuplicatedOf->bHasDuplicate)
+		{
+			DuplicatedOf->bHasDuplicate = true;
+			DuplicatedOf->GenerateLabelText();
+		}
 	}
 
 	// Generate the label to be used to display this operation in the operation tree
@@ -3378,7 +3388,7 @@ void FMutableCodeTreeElement::GenerateLabelText()
 	// MainLabel.InsertAt(0,RowStateIndex);
 
 	// Ignore the special case of operations of type "None"
-	if (MutableOperation > 0 && DuplicatedOf)
+	if (MutableOperation > 0 && bHasDuplicate)
 	{
 		MainLabel.Append(TEXT(" (duplicated)"));
 	}
diff --git a/Engine/Plugins/Mutable/Source/CustomizableObjectEditor/Private/MuCOE/SMutableCodeViewer.h b/Engine/Plugins/Mutable/Source/CustomizableObjectEditor/Private/MuCOE/SMutableCodeViewer.h
--- a/Engine/Plugins/Mutable/Source/CustomizableObjectEditor/Private/MuCOE/SMutableCodeViewer.h
+++ b/Engine/Plugins/Mutable/Source/CustomizableObjectEditor/Private/MuCOE/SMutableCodeViewer.h
@@ -907,6 +907,8 @@ public:
 	 * element then this could be part of the view (if inside view space).
 	 */
 	bool bIsExpanded = false;
+	
+	bool bHasDuplicate = false;
 
 private:
 	/** Represents as part of what state this element is currently part of. Will change for elements that are children
