[Mutable] Allow creating switches when there are at least 2 nested conditionals.

#Synced-CL 52761943

diff --git a/Engine/Plugins/Mutable/Source/MutableTools/Private/MuT/CodeOptimiser_Logic.cpp b/Engine/Plugins/Mutable/Source/MutableTools/Private/MuT/CodeOptimiser_Logic.cpp
--- a/Engine/Plugins/Mutable/Source/MutableTools/Private/MuT/CodeOptimiser_Logic.cpp
+++ b/Engine/Plugins/Mutable/Source/MutableTools/Private/MuT/CodeOptimiser_Logic.cpp
@@ -340,7 +340,7 @@ namespace UE::Mutable::Private
         {
             MUTABLE_CPUPROFILER_SCOPE(ConditionalToSwitch);
         	
-        	constexpr int32 MIN_CONDITIONS_TO_CREATE_SWITCH = 3;
+        	constexpr int32 MIN_CONDITIONS_TO_CREATE_SWITCH = 2;
         	
         	// Force optimization of the A branch of the conditional BO_AND operation
             ASTOp::Traverse_TopDown_Unique_Imprecise( roots, [&](Ptr<ASTOp>& CurrentOp)
