[mutable] Improved usability of the debugger by adding the option to color the tree using the operation data type.

#rnx
#Synced-CL 52778478

diff --git a/Engine/Plugins/Mutable/Source/CustomizableObjectEditor/Private/MuCOE/SMutableCodeViewer.cpp b/Engine/Plugins/Mutable/Source/CustomizableObjectEditor/Private/MuCOE/SMutableCodeViewer.cpp
--- a/Engine/Plugins/Mutable/Source/CustomizableObjectEditor/Private/MuCOE/SMutableCodeViewer.cpp
+++ b/Engine/Plugins/Mutable/Source/CustomizableObjectEditor/Private/MuCOE/SMutableCodeViewer.cpp
@@ -2,6 +2,8 @@
 
 #include "MuCOE/SMutableCodeViewer.h"
 
+#include "PropertyCustomizationHelpers.h"
+#include "SListViewSelectorDropdownMenu.h"
 #include "SMutableMaterialViewer.h"
 #include "Framework/MultiBox/MultiBoxBuilder.h"
 #include "Framework/Views/TableViewMetadata.h"
@@ -129,6 +131,12 @@ public:
 		return HighlightBoxDefaultColor;
 	}
 	
+
+	FSlateColor GetLabelColor() const
+	{
+		return RowItem->GetLabelColor();
+	}
+	
 	/** Method intended with the generation of the wanted objects for each column*/
 	virtual TSharedRef<SWidget> GenerateWidgetForColumn(const FName& ColumnName) override
 	{
@@ -164,7 +172,7 @@ public:
 					[
 						SNew(STextBlock)
 						.Text(FText::FromString(RowItem->MainLabel))
-						.ColorAndOpacity(RowItem->LabelColor)
+						.ColorAndOpacity(this, &SMutableCodeTreeRow::GetLabelColor)
 					]
 				]
 			];
@@ -348,6 +356,9 @@ void SMutableCodeViewer::Construct(const FArguments& InArgs, UCustomizableObject
 		Orientation(EOrientation::Orient_Vertical).
 		AlwaysShowScrollbar(false);
 	
+	// Ensure the color policy does have a valid default value.
+	check(EntryLabelColoringPolicyNames.IsValidIndex((int32)EntryLabelColoringPolicy));
+	
 	ChildSlot
 	[
 		SNew(SVerticalBox)
@@ -567,6 +578,35 @@ void SMutableCodeViewer::Construct(const FArguments& InArgs, UCustomizableObject
 					]
 				]
 				
+				// Label coloring policy
+				+ SVerticalBox::Slot()
+				.AutoHeight()
+				.HAlign(HAlign_Left)
+				[
+					SNew(SHorizontalBox)
+
+					+ SHorizontalBox::Slot()
+					.AutoWidth()
+					[
+						SNew(STextBlock)
+						.Text(LOCTEXT("OPLabelColoringStrategy", "Operation Coloring Strategy: "))
+			]
+
+					+ SHorizontalBox::Slot()
+					.AutoWidth()
+					.Padding(4,0)
+					[
+						SNew(SComboBox<FName>)
+						.OptionsSource(&EntryLabelColoringPolicyNames)
+						.InitiallySelectedItem(EntryLabelColoringPolicyNames[(int32)EntryLabelColoringPolicy])
+						.OnSelectionChanged(this, &SMutableCodeViewer::OnLabelColorPolicyChange)
+						.OnGenerateWidget(this, &SMutableCodeViewer::OnLabelColorPolicyGenerateWidget)
+						[
+							SNew(STextBlock)
+							.Text(this, &SMutableCodeViewer::GetCurrentLabelColoringPolicyText)
+						]
+					]
+				]
 			]
 			+ SSplitter::Slot()
 			.Value(0.75f)
@@ -932,7 +972,7 @@ FSlateColor SMutableCodeViewer::GetCurrentNavigationOpTypeColor() const
 {
 	check (CurrentlySelectedOperationTypeElement);
 
-	return CurrentlySelectedOperationTypeElement->OperationTextColor;
+	return CurrentlySelectedOperationTypeElement->bEntryColorWasOverriden ? CurrentlySelectedOperationTypeElement->OperationTextOverrideColor :  GetOperationColor(CurrentlySelectedOperationTypeElement->OperationType);
 }
 
 void SMutableCodeViewer::GenerateNavigationDropdownElements()
@@ -950,19 +990,16 @@ void SMutableCodeViewer::GenerateNavigationDropdownElements()
 		const UE::Mutable::Private::EOpType RepresentedType = ModelOperationTypes[OperationTypeIndex].Key;
 		const uint32 OperationTypeInstancesCount = ModelOperationTypes[OperationTypeIndex].Value;
 		
-		// Get the Color to be used on the text that will represent the operation on the dropdown
-		const FSlateColor OperationColor = ColorPerComputationalCost[StaticCast<uint8>(GetOperationTypeComputationalCost(RepresentedType))];
-		
 		// Generate an element to be used by the ComboBox handling the selection of the type to be used during navigation
-		TSharedPtr<FMutableOperationElement> OperationElement = MakeShared<FMutableOperationElement>(RepresentedType, OperationTypeName,OperationTypeInstancesCount,OperationColor);
+		TSharedPtr<FMutableOperationElement> OperationElement = MakeShared<FMutableOperationElement>(RepresentedType, OperationTypeName, OperationTypeInstancesCount);
 		FoundModelOperationTypeElements.Add(OperationElement);
 	}
 
 	// Add an entry for the NONE type of operation
 	{
 		const FText EntryName = FText::FromString("NONE");
-		const FSlateColor EntryColor = ColorPerComputationalCost[StaticCast<uint8>(EOperationComputationalCost::Standard)];
-		NoneOperationEntry = MakeShared<FMutableOperationElement>(UE::Mutable::Private::EOpType::NONE,EntryName,0,EntryColor);
+		
+		NoneOperationEntry = MakeShared<FMutableOperationElement>(UE::Mutable::Private::EOpType::NONE,EntryName,0);
 
 		// @warn While not visible this element must be part of the collection for the ComboBox to be able to work
 		// properly
@@ -972,8 +1009,12 @@ void SMutableCodeViewer::GenerateNavigationDropdownElements()
 	// Add an extra operation type that will represent the constant resource based navigation type
 	{
 		const FText EntryName = FText::FromString("Selected Constant");	
+		TSharedRef<FMutableOperationElement> NewMutableOperationElement = MakeShared<FMutableOperationElement>(UE::Mutable::Private::EOpType::NONE, EntryName, 0);
+
 		const FSlateColor EntryColor = FSlateColor(FLinearColor(0.35f ,0.35f,1.0f,1));
-		ConstantBasedNavigationEntry = MakeShared<FMutableOperationElement>(UE::Mutable::Private::EOpType::NONE,EntryName,0,EntryColor);
+		NewMutableOperationElement->SetEntryColorOverride(EntryColor);
+		
+		ConstantBasedNavigationEntry = NewMutableOperationElement;
 		
 		// @warn While not visible this element must be part of the collection for the ComboBox to be able to work
 		// properly
@@ -981,18 +1022,30 @@ void SMutableCodeViewer::GenerateNavigationDropdownElements()
 	}
 }
 
+
 TSharedRef<SWidget> SMutableCodeViewer::OnGenerateOpNavigationDropDownWidget(
-	TSharedPtr<const FMutableOperationElement> MutableOperationElement) const
+	TSharedPtr<const FMutableOperationElement> InMutableOperationElement) const
+{
+	if (!InMutableOperationElement->bEntryColorWasOverriden)
 {
 	TSharedRef<STextBlock> NewSlateObject = SNew(STextBlock)
-		.Text(MutableOperationElement->OperationTypeText)
-		.ColorAndOpacity(MutableOperationElement->OperationTextColor);
+		.Text(InMutableOperationElement->OperationTypeText)
+		.ColorAndOpacity(this, &SMutableCodeViewer::GetOperationColor, InMutableOperationElement->OperationType)
+		.Visibility(InMutableOperationElement->GetEntryVisibility());
 
-	// Set the visibility type for the UI object (currently will be hidden for the NONE type)
-	NewSlateObject->SetVisibility(MutableOperationElement->GetEntryVisibility());
+		return NewSlateObject;
+	}
+	else
+	{
+		TSharedRef<STextBlock> NewSlateObject = SNew(STextBlock)
+		.Text(InMutableOperationElement->OperationTypeText)
+		.ColorAndOpacity(InMutableOperationElement->OperationTextOverrideColor)
+		.Visibility(InMutableOperationElement->GetEntryVisibility());
 	
 	return NewSlateObject;
 }
+}
+
 
 void SMutableCodeViewer::OnNavigationSelectedOperationChanged(
 	TSharedPtr<const FMutableOperationElement, ESPMode::ThreadSafe> MutableOperationElement, ESelectInfo::Type Arg)
@@ -1320,14 +1373,12 @@ void SMutableCodeViewer::GenerateAllTreeElements()
 		const UE::Mutable::Private::FProgram::FState& State = Program.States[StateIndex];
 		const FString Caption = FString::Printf( TEXT("state [%s]"), *State.Name );
 
-		const FSlateColor LabelColor = ColorPerComputationalCost[StaticCast<uint8>(GetOperationTypeComputationalCost(Program.GetOpType(State.Root)))];
-
 		// Locate the "original" tree element :
 		// This may happen if for some reason the state is duplicated (should never happen)
 		const TSharedPtr<FMutableCodeTreeElement>* MainItemPtr = MainItemPerOp.Find(State.Root);
 
 		// Create a new root element and add it to the collection of root nodes
-		TSharedPtr<FMutableCodeTreeElement> RootNodeElement = MakeShareable(new FMutableCodeTreeElement(ItemCache.Num(),StateIndex, MutableModel, State.Root, Caption,LabelColor, MainItemPtr));
+		TSharedPtr<FMutableCodeTreeElement> RootNodeElement = MakeShareable(new FMutableCodeTreeElement(ItemCache.Num(),StateIndex, MutableModel, State.Root, Caption, SharedThis(this), Program.GetOpType(State.Root), MainItemPtr));
 		RootNodes.Add(RootNodeElement);
 		
 		// Add the element to the cache so we keep the indices straight.
@@ -1371,10 +1422,7 @@ void SMutableCodeViewer::GenerateElement(const int32& InStateIndex, UE::Mutable:
 			// Locate the "original" tree element
 			const TSharedPtr<FMutableCodeTreeElement>* MainItemPtr = MainItemPerOp.Find(ChildAddress);
 
-			// Provide the color this element should be using for the displayed text
-			const FSlateColor LabelColor = ColorPerComputationalCost[StaticCast<uint8>(GetOperationTypeComputationalCost(InProgram.GetOpType(ChildAddress)))];
-
-			const TSharedPtr<FMutableCodeTreeElement> Item = MakeShareable(new FMutableCodeTreeElement(ItemCache.Num(),InStateIndex, MutableModel, ChildAddress, Caption, LabelColor, MainItemPtr));
+			const TSharedPtr<FMutableCodeTreeElement> Item = MakeShareable(new FMutableCodeTreeElement(ItemCache.Num(),InStateIndex, MutableModel, ChildAddress, Caption, SharedThis(this), InProgram.GetOpType(ChildAddress), MainItemPtr));
 
 			// Cache this element for later access
 			ItemCache.Add(Key, Item);
@@ -1954,6 +2002,7 @@ void SMutableCodeViewer::GenerateElement(const int32& InStateIndex, UE::Mutable:
 	}
 }
 
+
 SMutableCodeViewer::EOperationComputationalCost SMutableCodeViewer::GetOperationTypeComputationalCost(UE::Mutable::Private::EOpType OperationType) const
 {
 	if (VeryExpensiveOperationTypes.Contains(OperationType))
@@ -1970,6 +2019,116 @@ SMutableCodeViewer::EOperationComputationalCost SMutableCodeViewer::GetOperation
 	}
 }
 
+
+FSlateColor SMutableCodeViewer::GetOperationDataTypeColor(UE::Mutable::Private::EDataType DataType) const
+{
+	FSlateColor Output = FSlateColor(FColor::White);
+	
+	switch (DataType)
+	{
+		case UE::Mutable::Private::EDataType::Bool:
+			Output = FSlateColor(FLinearColor(1,0.69,0.69, 1));
+			break;
+		case UE::Mutable::Private::EDataType::Int:
+			Output = FSlateColor(FLinearColor(0.5,0.24,0.45, 1));
+			break;
+		case UE::Mutable::Private::EDataType::Scalar:
+			Output = FSlateColor(FLinearColor(1,0.71,0, 1));
+			break;
+		case UE::Mutable::Private::EDataType::Color:
+			Output = FSlateColor(FLinearColor(0.64,0.73,0.83, 1));
+			break;
+		case UE::Mutable::Private::EDataType::Image:
+			Output = FSlateColor(FLinearColor(0.75,0,0.12, 1));
+			break;
+		case UE::Mutable::Private::EDataType::Layout:
+			Output = FSlateColor(FLinearColor(0.80,0.75,0.38, 1));
+			break;
+		case UE::Mutable::Private::EDataType::Mesh:
+			Output = FSlateColor(FLinearColor(0.50,0.43,0.71,1));
+			break;
+		case UE::Mutable::Private::EDataType::Instance:
+			Output = FSlateColor(FLinearColor(0,0.48,0.20,1));
+			break;
+		case UE::Mutable::Private::EDataType::Projector:
+			Output = FSlateColor(FLinearColor(0.96,0.46,0.55,1));
+			break;
+		case UE::Mutable::Private::EDataType::String:
+			Output = FSlateColor(FLinearColor(0,0.32,0.53,1));
+			break;
+		case UE::Mutable::Private::EDataType::ExtensionData:
+			Output = FSlateColor(FLinearColor(1,0.55,0,1));
+			break;
+		case UE::Mutable::Private::EDataType::Matrix:
+			Output = FSlateColor(FLinearColor(0.69,0.15,0.31,1));
+			break;
+		case UE::Mutable::Private::EDataType::Material:
+			Output = FSlateColor(FLinearColor(0.95,0.78,0,1));
+			break;
+		case UE::Mutable::Private::EDataType::InstancedStruct:
+			Output = FSlateColor(FLinearColor(0.49,0.09,0.05,1));
+			break;
+		case UE::Mutable::Private::EDataType::SkeletalMesh:
+			Output = FSlateColor(FLinearColor(0.94,0.22,0.84, 1));
+			break;
+	}
+	
+	// Other colors that could be used if one of those does show to not be ideal:
+	// Output = FSlateColor(FLinearColor(0.34,0.19,0.08, 1));
+	// Output = FSlateColor(FLinearColor(0.57,0.66,0,1));	
+	// Output = FSlateColor(FLinearColor(0.13,0.17,0.08, 1));
+	
+	return Output;
+}
+
+
+
+FSlateColor SMutableCodeViewer::GetOperationColor(const UE::Mutable::Private::EOpType OperationType) const
+{
+	if (EntryLabelColoringPolicy == EEntryLabelColoringPolicy::BY_COST)
+	{
+		return ColorPerComputationalCost[StaticCast<uint8>(GetOperationTypeComputationalCost(OperationType))];
+	}
+	else if (EntryLabelColoringPolicy == EEntryLabelColoringPolicy::BY_DATATYPE)
+	{
+		return GetOperationDataTypeColor(UE::Mutable::Private::GetOpDataType(OperationType));
+	}
+
+	checkNoEntry();
+	return FLinearColor::White;
+}
+
+
+void SMutableCodeViewer::OnLabelColorPolicyChange(FName PolicyName, ESelectInfo::Type Arg)
+{
+	check(EntryLabelColoringPolicyNames.Num() == (int32)EEntryLabelColoringPolicy::COUNT);
+	
+	const int32 IndexOnNameArray = EntryLabelColoringPolicyNames.Find(PolicyName);
+	check(IndexOnNameArray != INDEX_NONE);
+	check(EntryLabelColoringPolicyNames.IsValidIndex(IndexOnNameArray))
+
+	EntryLabelColoringPolicy = StaticCast<EEntryLabelColoringPolicy>(IndexOnNameArray);
+}
+
+
+TSharedRef<SWidget> SMutableCodeViewer::OnLabelColorPolicyGenerateWidget(FName PolicyName)
+{
+	TSharedRef<STextBlock> NewSlateObject = SNew(STextBlock)
+		.Text(FText::FromName(PolicyName));
+
+	return NewSlateObject;
+}
+
+
+FText SMutableCodeViewer::GetCurrentLabelColoringPolicyText() const
+{
+	check(EntryLabelColoringPolicyNames.IsValidIndex((int32)EntryLabelColoringPolicy));
+	const FName CurrentPolicyName = EntryLabelColoringPolicyNames[(int32)EntryLabelColoringPolicy];
+	
+	return FText::FromName(CurrentPolicyName);
+}
+
+
 #pragma endregion 
 
 
@@ -2541,6 +2700,7 @@ void SMutableCodeViewer::TreeExpandSelected()
 	TreeExpandElements(SelectedItems,true);
 }
 
+
 void SMutableCodeViewer::TreeExpandUnique()
 {
 	// Expand the tree from the root and do not expand the duplicated elements
@@ -3160,18 +3320,13 @@ void SMutableCodeViewer::PreviewMutableShape(const UE::Mutable::Private::FShape*
 }
 
 
-FMutableCodeTreeElement::FMutableCodeTreeElement(int32 InIndexOnTree, const int32& InMutableStateIndex, const TSharedPtr<UE::Mutable::Private::FModel, ESPMode::ThreadSafe>& InModel, UE::Mutable::Private::FOperation::ADDRESS InOperation, const FString& InCaption, const FSlateColor InLabelColor, const TSharedPtr<FMutableCodeTreeElement>* InDuplicatedOf)
+FMutableCodeTreeElement::FMutableCodeTreeElement(int32 InIndexOnTree, const int32& InMutableStateIndex, const TSharedPtr<UE::Mutable::Private::FModel, ESPMode::ThreadSafe>& InModel, UE::Mutable::Private::FOperation::ADDRESS InOperation, const FString& InCaption, TWeakPtr<SMutableCodeViewer> InCodeViewer, UE::Mutable::Private::EOpType InOperationType, const TSharedPtr<FMutableCodeTreeElement>* InDuplicatedOf)
 {
 	MutableModel = InModel;
 	MutableOperation = InOperation;
 	Caption = InCaption;
-	LabelColor = InLabelColor;
-	
-	// Use a special color to denote "none" entries instead of the provided one
-	if (MutableOperation == 0)
-	{
-		LabelColor = FColor(FColorList::DimGrey);
-	}
+	HostCodeViewer = InCodeViewer;
+	OperationType = InOperationType;
 	
 	IndexOnTree = InIndexOnTree;
 
@@ -3186,7 +3341,7 @@ FMutableCodeTreeElement::FMutableCodeTreeElement(int32 InIndexOnTree, const int3
 		{
 			DuplicatedOf->bHasDuplicate = true;
 			DuplicatedOf->GenerateLabelText();
-		}
+	}
 	}
 
 	// Generate the label to be used to display this operation in the operation tree
@@ -3237,7 +3392,6 @@ void FMutableCodeTreeElement::SetElementCurrentState(const int32& InStateIndex)
 void FMutableCodeTreeElement::GenerateLabelText()
 {
 	const UE::Mutable::Private::FProgram& Program = MutableModel->GetProgram();
-	const UE::Mutable::Private::EOpType OperationType = Program.GetOpType(MutableOperation);
 	FString OpName = UE::Mutable::Private::s_opNames[static_cast<int32>(OperationType)];
 	OpName.TrimEndInline();
 
@@ -3401,4 +3555,21 @@ int32 FMutableCodeTreeElement::GetStateIndex() const
 }
 
 
+FSlateColor FMutableCodeTreeElement::GetLabelColor() const
+{
+	if (MutableOperation == 0)
+	{
+		return ZeroOperationColor;
+	}
+	
+	if (TSharedPtr<SMutableCodeViewer> PinnedViewer = HostCodeViewer.Pin())
+	{
+		return PinnedViewer->GetOperationColor(OperationType);
+	}
+	
+	// The HostCodeViewer may no longer exist, the color returned will not be used.
+	return FColor::Black;
+}
+
+
 #undef LOCTEXT_NAMESPACE
diff --git a/Engine/Plugins/Mutable/Source/CustomizableObjectEditor/Private/MuCOE/SMutableCodeViewer.h b/Engine/Plugins/Mutable/Source/CustomizableObjectEditor/Private/MuCOE/SMutableCodeViewer.h
--- a/Engine/Plugins/Mutable/Source/CustomizableObjectEditor/Private/MuCOE/SMutableCodeViewer.h
+++ b/Engine/Plugins/Mutable/Source/CustomizableObjectEditor/Private/MuCOE/SMutableCodeViewer.h
@@ -218,6 +218,23 @@ private:
 	 * Operations computational cost reference collections
 	 */
 
+	enum class EEntryLabelColoringPolicy : int32
+	{
+		BY_COST = 0,
+		BY_DATATYPE,
+
+		COUNT,
+	};
+
+	const TArray<FName> EntryLabelColoringPolicyNames
+	{
+		"BY_COST",
+		"BY_DATATYPE",
+	};
+	
+	/** Determines what should be used to paint the labels */
+	EEntryLabelColoringPolicy EntryLabelColoringPolicy = EEntryLabelColoringPolicy::BY_DATATYPE;
+	
 	/** Collection with all very expensive to run operation types */
 	const TArray<UE::Mutable::Private::EOpType> VeryExpensiveOperationTypes
 	{
@@ -264,7 +281,16 @@ private:
 	 */
 	EOperationComputationalCost GetOperationTypeComputationalCost(UE::Mutable::Private::EOpType OperationType) const;
 	
+	FSlateColor GetOperationDataTypeColor(UE::Mutable::Private::EDataType DataType) const;
 	
+public:
+	FSlateColor GetOperationColor(const UE::Mutable::Private::EOpType OperationType) const;
+
+private:
+	void OnLabelColorPolicyChange(FName PolicyName, ESelectInfo::Type Arg);
+	TSharedRef<SWidget>  OnLabelColorPolicyGenerateWidget(FName PolicyName);
+	FText GetCurrentLabelColoringPolicyText() const;
+
 	/*
 	 * Navigation : Operation type navigation Type selection object
 	 */
@@ -322,7 +348,7 @@ private:
 	/** Callback invoked by the ComboBox used for displaying and selecting operation types for navigation. it gets invoked
 	 * each time the slate object requires to draw a line representing one of the elements set on FoundModelOperationTypeElements
 	 */
-	TSharedRef<SWidget> OnGenerateOpNavigationDropDownWidget(TSharedPtr<const FMutableOperationElement> MutableOperationElement) const;
+	TSharedRef<SWidget> OnGenerateOpNavigationDropDownWidget(TSharedPtr<const FMutableOperationElement> InMutableOperationElement) const;
 
 	/** Callback invoked each time the selected operation on our navigation slate changes. It can change due to UI interaction
 	 * or also due to direct change by invoking the SetSelectedOption on the SComboBox TargetedTypeSelector
@@ -825,10 +851,10 @@ private:
 class FMutableOperationElement : public TSharedFromThis<FMutableOperationElement>
 {
 public:
-	FMutableOperationElement(UE::Mutable::Private::EOpType InOperationType,FText OperationTypeName,uint32 OperationTypeInstanceCount,FSlateColor OperationColor)
+	
+	FMutableOperationElement(UE::Mutable::Private::EOpType InOperationType,FText OperationTypeName,uint32 OperationTypeInstanceCount)
 	{
 		OperationType = InOperationType;
-		OperationTextColor = OperationColor;
 
 		// Show the amount of instances of the operation type is found on model.
 		// @note This if block will not be triggered when working with an operation type that is not present on the model
@@ -841,12 +867,25 @@ public:
 		OperationTypeText = FText::Format(INVTEXT("{0}{1}"), OperationTypeName,OperationsCountText);
 	}
 
-public:
+	
+	/** Set the color to be used by this entry so it gets used instead of the one defined by the coloring policy 
+	 * @param InEntryColorOverride The color to be used for this entry.
+	 */
+	void SetEntryColorOverride(const FSlateColor& InEntryColorOverride)
+	{
+		bEntryColorWasOverriden = true;
+		OperationTextOverrideColor = InEntryColorOverride;
+	}
+
+	
 	UE::Mutable::Private::EOpType OperationType;
+	
 	FText OperationTypeText;
-	FSlateColor OperationTextColor;
+	
+	FSlateColor OperationTextOverrideColor;
+	bool bEntryColorWasOverriden = false;
 
-	/** Get the type of visibility an slate representing this objects should have. */
+	/** Get the type of visibility a slate representing this objects should have. */
 	EVisibility GetEntryVisibility() const
 	{
 		return OperationType == UE::Mutable::Private::EOpType::NONE ? EVisibility::Collapsed : EVisibility::Visible;
@@ -859,7 +898,7 @@ public:
 class FMutableCodeTreeElement : public TSharedFromThis<FMutableCodeTreeElement>
 {
 public:
-	FMutableCodeTreeElement(int32 InIndexOnTree, const int32& InMutableStateIndex, const TSharedPtr<UE::Mutable::Private::FModel, ESPMode::ThreadSafe>& InModel, UE::Mutable::Private::FOperation::ADDRESS InOperation, const FString& InCaption, const FSlateColor InLabelColor, const TSharedPtr<FMutableCodeTreeElement>* InDuplicatedOf = nullptr);
+	FMutableCodeTreeElement(int32 InIndexOnTree, const int32& InMutableStateIndex, const TSharedPtr<UE::Mutable::Private::FModel, ESPMode::ThreadSafe>& InModel, UE::Mutable::Private::FOperation::ADDRESS InOperation, const FString& InCaption, TWeakPtr<SMutableCodeViewer> InCodeViewer, UE::Mutable::Private::EOpType InOperationType, const TSharedPtr<FMutableCodeTreeElement>* InDuplicatedOf = nullptr);
 	
 	void SetElementCurrentState(const int32& InStateIndex);
 
@@ -867,6 +906,8 @@ public:
 	
 	int32 GetStateIndex() const;
 	
+	FSlateColor GetLabelColor() const;
+
 public:
 
 	/** Model this operation is part of */
@@ -878,9 +919,6 @@ public:
 	/** If this tree element is a duplicated of another op, this is the op. */
 	TSharedPtr<FMutableCodeTreeElement> DuplicatedOf;
 
-	/** The color to be used by the row representing this object */
-	FSlateColor LabelColor;
-
 	/** Label to be used to represent this operation in the tree */
 	FString MainLabel;
 
@@ -907,7 +945,7 @@ public:
 	 * element then this could be part of the view (if inside view space).
 	 */
 	bool bIsExpanded = false;
-	
+
 	bool bHasDuplicate = false;
 
 private:
@@ -919,4 +957,13 @@ private:
 
 	/** String representing additional data like the state this operation is from. It may be empty. */
 	FString Caption;
+
+	/** Code Viewer hosting the tree that uses this object. Will be used to evaluate the code to be displayed by the tree entry. */
+	TWeakPtr<SMutableCodeViewer> HostCodeViewer;
+
+	/** The type of the mutable operation. */
+	UE::Mutable::Private::EOpType OperationType;
+	
+	/** The color to be used by the row representing this object */
+	FSlateColor ZeroOperationColor = FColor(FColorList::DimGrey);
 };
