[mutable] Added Mutable Skeletal Mesh Merge node (editor and core)
- The node is currently not spawnable as there is no code generation code for it yet.

#rnx
#Synced-CL 52784087

diff --git a/Engine/Plugins/Mutable/Source/CustomizableObjectEditor/Private/MuCOE/EdGraphSchema_CustomizableObject.cpp b/Engine/Plugins/Mutable/Source/CustomizableObjectEditor/Private/MuCOE/EdGraphSchema_CustomizableObject.cpp
--- a/Engine/Plugins/Mutable/Source/CustomizableObjectEditor/Private/MuCOE/EdGraphSchema_CustomizableObject.cpp
+++ b/Engine/Plugins/Mutable/Source/CustomizableObjectEditor/Private/MuCOE/EdGraphSchema_CustomizableObject.cpp
@@ -117,6 +117,7 @@
 #include "Nodes/CONodeRemoveMesh.h"
 #include "Nodes/CONodeSkeletalMeshBreak.h"
 #include "Nodes/CONodeMutableSkeletalMeshMake.h"
+#include "Nodes/CONodeMutableSkeletalMeshMerge.h"
 #include "Nodes/CONodeTransformWithBone.h"
 #include "Nodes/CONodeTransformInMesh.h"
 #include "Nodes/CONodeSwitch.h"
@@ -648,6 +649,7 @@ void UEdGraphSchema_CustomizableObject::GetGraphContextActions(FGraphContextMenu
 			ContextMenuBuilder.CreateTemplateNode<UCONodeSkeletalMeshBreak>(),
 			ContextMenuBuilder.CreateTemplateNode<UCONodeMutableSkeletalMeshMake>(),
 			ContextMenuBuilder.CreateTemplateNode<UCONodeSkeletalMeshMake>(),
+			// ContextMenuBuilder.CreateTemplateNode<UCONodeMutableSkeletalMeshMerge>(),
 		};
 
 		AddNewNodeCategoryActionsFiltered(ExperimentalTemplateNodes, ContextMenuBuilder, NC_Experimental, 2, Filter);
diff --git a/Engine/Plugins/Mutable/Source/CustomizableObjectEditor/Private/MuCOE/GenerateMutableSource/GenerateMutableSourceSkeletalMesh.cpp b/Engine/Plugins/Mutable/Source/CustomizableObjectEditor/Private/MuCOE/GenerateMutableSource/GenerateMutableSourceSkeletalMesh.cpp
--- a/Engine/Plugins/Mutable/Source/CustomizableObjectEditor/Private/MuCOE/GenerateMutableSource/GenerateMutableSourceSkeletalMesh.cpp
+++ b/Engine/Plugins/Mutable/Source/CustomizableObjectEditor/Private/MuCOE/GenerateMutableSource/GenerateMutableSourceSkeletalMesh.cpp
@@ -5,6 +5,7 @@
 #include "GenerateMutableSource.h"
 #include "GenerateMutableSourceSurface.h"
 #include "MuCOE/GraphTraversal.h"
+#include "MuCOE/Nodes/CONodeMutableSkeletalMeshMerge.h"
 #include "MuCOE/Nodes/CustomizableObjectNodeComponentPassthroughMesh.h"
 #include "MuCOE/Nodes/CustomizableObjectNodeModifierBase.h"
 #include "MuCOE/Nodes/CustomizableObjectNodeSkeletalMeshParameter.h"
@@ -167,6 +168,40 @@ namespace UE::Mutable::Private
 		return {};
 	}
 
+	
+	Ptr<NodeSkeletalMeshMerge> GenerateMutableSkeletalMeshMerge(const UEdGraphPin* Pin, FMutableGraphGenerationContext& GenerationContext)
+	{
+		const UCONodeMutableSkeletalMeshMerge* Node = CastChecked<UCONodeMutableSkeletalMeshMerge>(Pin->GetOwningNode());
+		
+		Ptr<NodeSkeletalMeshMerge> Result = new NodeSkeletalMeshMerge();
+		
+		const UEdGraphPin* BaseMeshPin = Node->BaseSkeletalMeshPin.Get();
+		check(BaseMeshPin);
+		if (const UEdGraphPin* ConnectedBaseMeshPin = FollowInputPin(*BaseMeshPin))
+		{
+			Result->BaseSkeletalMesh = GenerateMutableSourceSkeletalMesh(ConnectedBaseMeshPin, GenerationContext);
+		}
+		else
+		{
+			const FText Message = LOCTEXT("CONodeSkeletalMeshMerge_NoInputConnectedWarning","No connections could be found to the Base Skeletal Mesh input pin.");
+			GenerationContext.Log(Message, Node, EMessageSeverity::Warning);
+		}
+		
+		const UEdGraphPin* ToAddMeshPin = Node->ToAddSkeletalMeshPin.Get();
+		check(ToAddMeshPin);
+		if (const UEdGraphPin* ConnectedToAddMeshPin = FollowInputPin(*ToAddMeshPin))
+		{
+			Result->ToAddSkeletalMesh = GenerateMutableSourceSkeletalMesh(ConnectedToAddMeshPin, GenerationContext);
+		}
+		else
+		{
+			const FText Message = LOCTEXT("CONodeSkeletalMeshMerge_NoInputConnectedWarning","No connections could be found to the To Add Skeletal Mesh input pin.");
+			GenerationContext.Log(Message, Node, EMessageSeverity::Warning);
+		}
+		
+		return Result;
+	}
+
 
 	Ptr<NodeSkeletalMesh> GenerateMutableSourceSkeletalMesh(const UEdGraphPin* Pin, FMutableGraphGenerationContext& GenerationContext)
 	{
@@ -199,6 +234,11 @@ namespace UE::Mutable::Private
 			Result = GenerateMutableSourceSkeletalMeshNew(Pin, GenerationContext);
 		}
 
+		else if (const UCONodeMutableSkeletalMeshMerge* SkeletalMeshMergeNode = Cast<UCONodeMutableSkeletalMeshMerge>(Node))
+		{
+			Result = GenerateMutableSkeletalMeshMerge(Pin, GenerationContext);
+		}
+		
 		else if (const UCustomizableObjectNodeComponentPassthroughMesh* ComponentPassthroughMesh = Cast<UCustomizableObjectNodeComponentPassthroughMesh>(Node))
 		{
 			if (const UEdGraphPin* ConnectedPin = FollowInputPin(*ComponentPassthroughMesh->SkeletalMeshPin.Get()))
diff --git a/Engine/Plugins/Mutable/Source/CustomizableObjectEditor/Private/MuCOE/GenerateMutableSource/GenerateMutableSourceSkeletalMesh.h b/Engine/Plugins/Mutable/Source/CustomizableObjectEditor/Private/MuCOE/GenerateMutableSource/GenerateMutableSourceSkeletalMesh.h
--- a/Engine/Plugins/Mutable/Source/CustomizableObjectEditor/Private/MuCOE/GenerateMutableSource/GenerateMutableSourceSkeletalMesh.h
+++ b/Engine/Plugins/Mutable/Source/CustomizableObjectEditor/Private/MuCOE/GenerateMutableSource/GenerateMutableSourceSkeletalMesh.h
@@ -4,6 +4,7 @@
 
 #include "MuR/Ptr.h"
 #include "MuT/NodeSkeletalMesh.h"
+#include "MuT/NodeSkeletalMeshMerge.h"
 #include "MuT/NodeSkeletalMeshNew.h"
 #include "MuT/NodeSkeletalMeshParameter.h"
 
@@ -19,4 +20,6 @@ namespace UE::Mutable::Private
 
 	Ptr<NodeSkeletalMeshNew> GenerateMutableSourceSkeletalMeshNew(const UEdGraphPin* Pin, FMutableGraphGenerationContext& GenerationContext);
 
+	Ptr<NodeSkeletalMeshMerge> GenerateMutableSkeletalMeshMerge(const UEdGraphPin* Pin, FMutableGraphGenerationContext& GenerationContext);
+	
 }
diff --git a/Engine/Plugins/Mutable/Source/CustomizableObjectEditor/Private/MuCOE/Nodes/CONodeMutableSkeletalMeshMake.cpp b/Engine/Plugins/Mutable/Source/CustomizableObjectEditor/Private/MuCOE/Nodes/CONodeMutableSkeletalMeshMake.cpp
--- a/Engine/Plugins/Mutable/Source/CustomizableObjectEditor/Private/MuCOE/Nodes/CONodeMutableSkeletalMeshMake.cpp
+++ b/Engine/Plugins/Mutable/Source/CustomizableObjectEditor/Private/MuCOE/Nodes/CONodeMutableSkeletalMeshMake.cpp
@@ -88,13 +88,13 @@ bool UCONodeMutableSkeletalMeshMake::CustomRemovePin(UEdGraphPin& Pin)
 
 FText UCONodeMutableSkeletalMeshMake::GetNodeTitle(ENodeTitleType::Type TitleType) const
 {
-	return LOCTEXT("Mutable_Skeletal_Mesh_Make_Node_Title", "Mutable Skeletal Mesh Make");
+	return LOCTEXT("CONodeMutableSkeletalMeshMake_NodeTitle", "Mutable Skeletal Mesh Make");
 }
 
 
 FLinearColor UCONodeMutableSkeletalMeshMake::GetNodeTitleColor() const
 {
-	return UEdGraphSchema_CustomizableObject::GetPinTypeColor(UEdGraphSchema_CustomizableObject::PC_SkeletalMesh_Passthrough);
+	return UEdGraphSchema_CustomizableObject::GetPinTypeColor(UEdGraphSchema_CustomizableObject::PC_SkeletalMesh);
 }
 
 #undef LOCTEXT_NAMESPACE
diff --git a/Engine/Plugins/Mutable/Source/CustomizableObjectEditor/Private/MuCOE/Nodes/CONodeMutableSkeletalMeshMerge.cpp b/Engine/Plugins/Mutable/Source/CustomizableObjectEditor/Private/MuCOE/Nodes/CONodeMutableSkeletalMeshMerge.cpp
new file mode 100644
--- /dev/null
+++ b/Engine/Plugins/Mutable/Source/CustomizableObjectEditor/Private/MuCOE/Nodes/CONodeMutableSkeletalMeshMerge.cpp
@@ -0,0 +1,40 @@
+// Copyright Epic Games, Inc. All Rights Reserved.
+
+#include "MuCOE/Nodes/CONodeMutableSkeletalMeshMerge.h"
+
+#include "MuCOE/EdGraphSchema_CustomizableObject.h"
+
+class UCustomizableObjectNodeRemapPins;
+
+#define LOCTEXT_NAMESPACE "CustomizableObjectEditor"
+
+
+FText UCONodeMutableSkeletalMeshMerge::GetNodeTitle(ENodeTitleType::Type TitleType) const
+{
+	return LOCTEXT("CONodeSkeletalMeshMerge_NodeTitle", "Mutable Skeletal Mesh Merge");
+}
+
+
+FLinearColor UCONodeMutableSkeletalMeshMerge::GetNodeTitleColor() const
+{
+	return UEdGraphSchema_CustomizableObject::GetPinTypeColor(UEdGraphSchema_CustomizableObject::PC_SkeletalMesh);
+}
+
+
+void UCONodeMutableSkeletalMeshMerge::AllocateDefaultPins(UCustomizableObjectNodeRemapPins* RemapPins)
+{
+	// Inputs
+	
+	const FText BaseSkeletalMeshPinName = LOCTEXT("CONodeSkeletalMeshMerge_BaseSkeletalMeshName", "Base Skeletal Mesh");
+	BaseSkeletalMeshPin = CustomCreatePin(EGPD_Input, UEdGraphSchema_CustomizableObject::PC_SkeletalMesh,  *BaseSkeletalMeshPinName.ToString(), BaseSkeletalMeshPinName);
+	
+	const FText ToAddSkeletalMeshPinName = LOCTEXT("CONodeSkeletalMeshMerge_ToAddSkeletalMeshName", "To Add Skeletal Mesh");
+	ToAddSkeletalMeshPin = CustomCreatePin(EGPD_Input, UEdGraphSchema_CustomizableObject::PC_SkeletalMesh, *ToAddSkeletalMeshPinName.ToString(), ToAddSkeletalMeshPinName);
+
+	// Output
+	
+	SkeletalMeshOutputPin = CustomCreatePin(EGPD_Output, UEdGraphSchema_CustomizableObject::PC_SkeletalMesh);
+}
+
+
+#undef LOCTEXT_NAMESPACE
\ No newline at end of file
diff --git a/Engine/Plugins/Mutable/Source/CustomizableObjectEditor/Private/MuCOE/Nodes/CONodeMutableSkeletalMeshMerge.h b/Engine/Plugins/Mutable/Source/CustomizableObjectEditor/Private/MuCOE/Nodes/CONodeMutableSkeletalMeshMerge.h
new file mode 100644
--- /dev/null
+++ b/Engine/Plugins/Mutable/Source/CustomizableObjectEditor/Private/MuCOE/Nodes/CONodeMutableSkeletalMeshMerge.h
@@ -0,0 +1,36 @@
+// Copyright Epic Games, Inc. All Rights Reserved.
+
+#pragma once
+
+#include "MuCOE/Nodes/CustomizableObjectNode.h"
+
+#include "CONodeMutableSkeletalMeshMerge.generated.h"
+
+
+UCLASS()
+class UCONodeMutableSkeletalMeshMerge : public UCustomizableObjectNode
+{
+public:
+	
+	GENERATED_BODY()
+	
+	// UEdGraphNode
+	virtual FText GetNodeTitle(ENodeTitleType::Type TitleType) const override;
+	virtual FLinearColor GetNodeTitleColor() const override;
+	
+protected:
+	
+	virtual void AllocateDefaultPins(UCustomizableObjectNodeRemapPins* RemapPins) override;
+	
+public:
+	
+	UPROPERTY()
+	FEdGraphPinReference BaseSkeletalMeshPin;
+	
+	UPROPERTY()
+	FEdGraphPinReference ToAddSkeletalMeshPin;
+	
+	UPROPERTY()
+	FEdGraphPinReference SkeletalMeshOutputPin;
+	
+};
diff --git a/Engine/Plugins/Mutable/Source/CustomizableObjectEditor/Private/MuCOE/SMutableGraphViewer.cpp b/Engine/Plugins/Mutable/Source/CustomizableObjectEditor/Private/MuCOE/SMutableGraphViewer.cpp
--- a/Engine/Plugins/Mutable/Source/CustomizableObjectEditor/Private/MuCOE/SMutableGraphViewer.cpp
+++ b/Engine/Plugins/Mutable/Source/CustomizableObjectEditor/Private/MuCOE/SMutableGraphViewer.cpp
@@ -200,6 +200,7 @@ static const char* MutableNodeNames[] =
 
 	"SkeletalMesh",
 	"SkeletalMeshNew",
+	"SkeletalMeshMerge",
 	"SkeletalMeshParameter",
 
 	"Component",
diff --git a/Engine/Plugins/Mutable/Source/MutableTools/Internal/MuT/Node.h b/Engine/Plugins/Mutable/Source/MutableTools/Internal/MuT/Node.h
--- a/Engine/Plugins/Mutable/Source/MutableTools/Internal/MuT/Node.h
+++ b/Engine/Plugins/Mutable/Source/MutableTools/Internal/MuT/Node.h
@@ -138,6 +138,7 @@ namespace UE::Mutable::Private
 
 			SkeletalMesh,
 			SkeletalMeshNew,
+			SkeletalMeshMerge,
 			SkeletalMeshParameter,
 			
 			Component,
diff --git a/Engine/Plugins/Mutable/Source/MutableTools/Internal/MuT/NodeSkeletalMeshMerge.h b/Engine/Plugins/Mutable/Source/MutableTools/Internal/MuT/NodeSkeletalMeshMerge.h
new file mode 100644
--- /dev/null
+++ b/Engine/Plugins/Mutable/Source/MutableTools/Internal/MuT/NodeSkeletalMeshMerge.h
@@ -0,0 +1,38 @@
+// Copyright Epic Games, Inc. All Rights Reserved.
+
+#pragma once
+
+#include "NodeSkeletalMesh.h"
+// #include "HAL/Platform.h"
+#include "MuR/Ptr.h"
+//#include "MuT/Node.h"
+// #include "MuT/NodeLOD.h"
+
+
+#define UE_API MUTABLETOOLS_API
+
+namespace UE::Mutable::Private
+{
+	class NodeSkeletalMeshMerge : public NodeSkeletalMesh
+	{
+	public:
+		
+		/** Mesh where we want to merge into the sections of the ToAddSkeletalMesh*/
+		Ptr<NodeSkeletalMesh> BaseSkeletalMesh;
+		
+		/** Skeletal Mesh whose section will be merged with the ones on the BaseSkeletalMesh*/
+		Ptr<NodeSkeletalMesh> ToAddSkeletalMesh;
+		
+		// Node interface
+		virtual const FNodeType* GetType() const override { return &StaticType; }
+		static const FNodeType* GetStaticType() { return &StaticType; }
+	
+	protected:
+		virtual ~NodeSkeletalMeshMerge() override {}
+
+	private:
+		static UE_API FNodeType StaticType;
+	};
+}
+
+#undef UE_API
diff --git a/Engine/Plugins/Mutable/Source/MutableTools/Private/MuT/Node.cpp b/Engine/Plugins/Mutable/Source/MutableTools/Private/MuT/Node.cpp
--- a/Engine/Plugins/Mutable/Source/MutableTools/Private/MuT/Node.cpp
+++ b/Engine/Plugins/Mutable/Source/MutableTools/Private/MuT/Node.cpp
@@ -120,6 +120,7 @@
 #include "MuT/NodeMeshRemoveMesh.h"
 #include "MuT/NodeMeshTransformInMesh.h"
 #include "MuT/NodeMeshTransformWithBone.h"
+#include "MuT/NodeSkeletalMeshMerge.h"
 
 
 namespace UE::Mutable::Private
@@ -169,6 +170,7 @@ namespace UE::Mutable::Private
 
 	FNodeType NodeSkeletalMesh::StaticType = FNodeType(Node::EType::SkeletalMesh, Node::GetStaticType());
 	FNodeType NodeSkeletalMeshNew::StaticType = FNodeType(Node::EType::SkeletalMeshNew, Node::GetStaticType());
+	FNodeType NodeSkeletalMeshMerge::StaticType = FNodeType(Node::EType::SkeletalMeshMerge, Node::GetStaticType());
 	FNodeType NodeSkeletalMeshParameter::StaticType = FNodeType(Node::EType::SkeletalMeshParameter, NodeSkeletalMesh::GetStaticType());
 
 	FNodeType NodeExtensionData::StaticType = FNodeType(Node::EType::ExtensionData, Node::GetStaticType());
